Copying //UE4/Dev-Sequencer to //UE4/Main

diff --git a/Engine/Source/Editor/MovieSceneTools/Private/TrackEditors/SubTrackEditor.cpp b/Engine/Source/Editor/MovieSceneTools/Private/TrackEditors/SubTrackEditor.cpp
--- a/Engine/Source/Editor/MovieSceneTools/Private/TrackEditors/SubTrackEditor.cpp
+++ b/Engine/Source/Editor/MovieSceneTools/Private/TrackEditors/SubTrackEditor.cpp
@@ -9,7 +9,8 @@
 #include "ContentBrowserModule.h"
 #include "SequencerUtilities.h"
 #include "SequencerSectionPainter.h"
-
+#include "ISequenceRecorder.h"
+#include "SequenceRecorderSettings.h"
 
 namespace SubTrackEditorConstants
 {
@@ -62,7 +63,42 @@ public:
 
 	virtual FText GetSectionTitle() const override
 	{
-		return FText::FromString(SectionObject.GetSequence()->GetName());
+		if(SectionObject.GetSequence())
+		{
+			return FText::FromString(SectionObject.GetSequence()->GetName());
+		}
+		else if(UMovieSceneSubSection::GetRecordingSection() == &SectionObject)
+		{
+			const FString& ActorToRecord = UMovieSceneSubSection::GetActorToRecord();
+
+			ISequenceRecorder& SequenceRecorder = FModuleManager::GetModuleChecked<ISequenceRecorder>("SequenceRecorder");
+			if(SequenceRecorder.IsRecording())
+			{
+				if(ActorToRecord.Len() > 0)
+				{
+					return FText::Format(LOCTEXT("RecordingIndicatorWithActor", "Sequence Recording for \"{0}\""), FText::FromString(UMovieSceneSubSection::GetActorToRecord()));
+				}
+				else
+				{
+					return LOCTEXT("RecordingIndicator", "Sequence Recording");
+				}
+			}
+			else
+			{
+				if(ActorToRecord.Len() > 0)
+				{
+					return FText::Format(LOCTEXT("RecordingPendingIndicatorWithActor", "Sequence Recording Pending for \"{0}\""), FText::FromString(UMovieSceneSubSection::GetActorToRecord()));
+				}
+				else
+				{
+					return LOCTEXT("RecordingPendingIndicator", "Sequence Recording Pending");
+				}
+			}
+		}
+		else
+		{
+			return LOCTEXT("InvalidSequence", "No Sequence Selected");
+		}
 	}
 	
 	virtual int32 OnPaintSection( FSequencerSectionPainter& InPainter ) const override
@@ -81,13 +117,24 @@ public:
 			? ESlateDrawEffect::None
 			: ESlateDrawEffect::DisabledEffect;
 
-		UMovieScene* MovieScene = SectionObject.GetSequence()->GetMovieScene();
-		const FFloatRange& PlaybackRange = MovieScene->GetPlaybackRange();
+		UMovieScene* MovieScene = nullptr;
+		FFloatRange PlaybackRange;
+		if(SectionObject.GetSequence() != nullptr)
+		{
+			MovieScene = SectionObject.GetSequence()->GetMovieScene();
+			PlaybackRange = MovieScene->GetPlaybackRange();
+		}
+		else
+		{
+			UMovieSceneTrack* MovieSceneTrack = CastChecked<UMovieSceneTrack>(SectionObject.GetOuter());
+			MovieScene = CastChecked<UMovieScene>(MovieSceneTrack->GetOuter());
+			PlaybackRange = MovieScene->GetPlaybackRange();
+		}
 
 		// add box for the working size
 		const float StartOffset = SectionObject.TimeScale * SectionObject.StartOffset;
 		const float WorkingStart = -SectionObject.TimeScale * PlaybackRange.GetLowerBoundValue() - StartOffset;
-		const float WorkingSize = SectionObject.TimeScale * MovieScene->GetEditorData().WorkingRange.Size<float>();
+		const float WorkingSize = SectionObject.TimeScale * (MovieScene != nullptr ? MovieScene->GetEditorData().WorkingRange.Size<float>() : 1.0f);
 
 		FSlateDrawElement::MakeBox(
 			InPainter.DrawElements,
@@ -173,18 +220,37 @@ public:
 		}
 
 		// add number of tracks within the section's sequence
-		int32 NumTracks = MovieScene->GetPossessableCount() + MovieScene->GetSpawnableCount() + MovieScene->GetMasterTracks().Num();
+		if(SectionObject.GetSequence() != nullptr && MovieScene != nullptr)
+		{
+			int32 NumTracks = MovieScene->GetPossessableCount() + MovieScene->GetSpawnableCount() + MovieScene->GetMasterTracks().Num();
 
-		FSlateDrawElement::MakeText(
-			InPainter.DrawElements,
-			++LayerId,
-			InPainter.SectionGeometry.ToOffsetPaintGeometry(FVector2D(5.0f, 32.0f)),
-			FText::Format(LOCTEXT("NumTracksFormat", "{0} track(s)"), FText::AsNumber(NumTracks)),
-			FEditorStyle::GetFontStyle("NormalFont"),
-			InPainter.SectionClippingRect,
-			DrawEffects,
-			FColor(200, 200, 200)
-		);
+			FSlateDrawElement::MakeText(
+				InPainter.DrawElements,
+				++LayerId,
+				InPainter.SectionGeometry.ToOffsetPaintGeometry(FVector2D(11.0f, 32.0f)),
+				FText::Format(LOCTEXT("NumTracksFormat", "{0} track(s)"), FText::AsNumber(NumTracks)),
+				FEditorStyle::GetFontStyle("NormalFont"),
+				InPainter.SectionClippingRect,
+				DrawEffects,
+				FColor(200, 200, 200)
+			);
+		}
+		else
+		{
+			// if we are primed for recording, display where we will create the recording
+			FString Path = SectionObject.GetTargetPathToRecordTo() / SectionObject.GetTargetSequenceName();
+
+			FSlateDrawElement::MakeText(
+				InPainter.DrawElements,
+				++LayerId,
+				InPainter.SectionGeometry.ToOffsetPaintGeometry(FVector2D(11.0f, 32.0f)),
+				FText::Format(LOCTEXT("RecordingDestination", "Target: \"{0}\""), FText::FromString(Path)),
+				FEditorStyle::GetFontStyle("NormalFont"),
+				InPainter.SectionClippingRect,
+				DrawEffects,
+				FColor(200, 200, 200)
+			);
+		}
 
 		return LayerId;
 	}
@@ -193,7 +259,10 @@ public:
 	{
 		if( MouseEvent.GetEffectingButton() == EKeys::LeftMouseButton )
 		{
-			Sequencer.Pin()->FocusSequenceInstance(SectionObject);
+			if (SectionObject.GetSequence())
+			{
+				Sequencer.Pin()->FocusSequenceInstance(SectionObject);
+			}
 		}
 
 		return FReply::Handled();
@@ -372,24 +441,40 @@ TSharedRef<SWidget> FSubTrackEditor::HandleAddSubSequenceComboButtonGetMenuConte
 {
 	FMenuBuilder MenuBuilder(true, nullptr);
 
-	FAssetPickerConfig AssetPickerConfig;
+	MenuBuilder.BeginSection(TEXT("RecordSequence"), LOCTEXT("RecordSequence", "Record Sequence"));
 	{
-		AssetPickerConfig.OnAssetSelected = FOnAssetSelected::CreateRaw( this, &FSubTrackEditor::HandleAddSubSequenceComboButtonMenuEntryExecute);
-		AssetPickerConfig.bAllowNullSelection = false;
-		AssetPickerConfig.InitialAssetViewType = EAssetViewType::Tile;
-		AssetPickerConfig.Filter.ClassNames.Add(TEXT("LevelSequence"));
+		MenuBuilder.AddMenuEntry(
+			LOCTEXT("RecordNewSequence", "Record New Sequence"), 
+			LOCTEXT("RecordNewSequence_ToolTip", "Record a new level seqeuence into this sub-track from gameplay/simulation etc.\nThis only primes the track for recording. Click the record button to begin recording into this track once primed.\nOnly one sequence can be recorded at a time."), 
+			FSlateIcon(), 
+			FUIAction(
+				FExecuteAction::CreateSP(this, &FSubTrackEditor::HandleRecordNewSequence),
+				FCanExecuteAction::CreateSP(this, &FSubTrackEditor::CanRecordNewSequence)));
 	}
+	MenuBuilder.EndSection();
+
+	MenuBuilder.BeginSection(TEXT("ChooseSequence"), LOCTEXT("ChooseSequence", "Choose Sequence"));
+	{
+		FAssetPickerConfig AssetPickerConfig;
+		{
+			AssetPickerConfig.OnAssetSelected = FOnAssetSelected::CreateRaw( this, &FSubTrackEditor::HandleAddSubSequenceComboButtonMenuEntryExecute);
+			AssetPickerConfig.bAllowNullSelection = false;
+			AssetPickerConfig.InitialAssetViewType = EAssetViewType::Tile;
+			AssetPickerConfig.Filter.ClassNames.Add(TEXT("LevelSequence"));
+		}
 
-	FContentBrowserModule& ContentBrowserModule = FModuleManager::Get().LoadModuleChecked<FContentBrowserModule>(TEXT("ContentBrowser"));
+		FContentBrowserModule& ContentBrowserModule = FModuleManager::Get().LoadModuleChecked<FContentBrowserModule>(TEXT("ContentBrowser"));
 
-	TSharedPtr<SBox> MenuEntry = SNew(SBox)
-		.WidthOverride(300.0f)
-		.HeightOverride(300.f)
-		[
-			ContentBrowserModule.Get().CreateAssetPicker(AssetPickerConfig)
-		];
+		TSharedPtr<SBox> MenuEntry = SNew(SBox)
+			.WidthOverride(300.0f)
+			.HeightOverride(300.f)
+			[
+				ContentBrowserModule.Get().CreateAssetPicker(AssetPickerConfig)
+			];
 
-	MenuBuilder.AddWidget(MenuEntry.ToSharedRef(), FText::GetEmpty(), true);
+		MenuBuilder.AddWidget(MenuEntry.ToSharedRef(), FText::GetEmpty(), true);
+	}
+	MenuBuilder.EndSection();
 
 	return MenuBuilder.MakeWidget();
 }
@@ -431,5 +516,28 @@ bool FSubTrackEditor::HandleSequenceAdded(float KeyTime, UMovieSceneSequence* Se
 	return true;
 }
 
+bool FSubTrackEditor::CanRecordNewSequence() const
+{
+	return !UMovieSceneSubSection::IsSetAsRecording();
+}
+
+void FSubTrackEditor::HandleRecordNewSequence()
+{
+	FSlateApplication::Get().DismissAllMenus();
+
+	AnimatablePropertyChanged( FOnKeyProperty::CreateRaw( this, &FSubTrackEditor::HandleRecordNewSequenceInternal) );
+}
+
+bool FSubTrackEditor::HandleRecordNewSequenceInternal(float KeyTime)
+{
+	auto SubTrack = FindOrCreateMasterTrack<UMovieSceneSubTrack>().Track;
+	UMovieSceneSubSection* Section = SubTrack->AddSequenceToRecord();
+
+	// @todo: we could default to the same directory as a parent sequence, or the last sequence recorded. Lots of options!
+	Section->SetTargetSequenceName(GetDefault<USequenceRecorderSettings>()->SequenceName);
+	Section->SetTargetPathToRecordTo(GetDefault<USequenceRecorderSettings>()->SequenceRecordingBasePath.Path);
+
+	return true;
+}
 
 #undef LOCTEXT_NAMESPACE
diff --git a/Engine/Source/Editor/MovieSceneTools/Private/TrackEditors/SubTrackEditor.h b/Engine/Source/Editor/MovieSceneTools/Private/TrackEditors/SubTrackEditor.h
--- a/Engine/Source/Editor/MovieSceneTools/Private/TrackEditors/SubTrackEditor.h
+++ b/Engine/Source/Editor/MovieSceneTools/Private/TrackEditors/SubTrackEditor.h
@@ -72,4 +72,13 @@ private:
 
 	/** Callback for AnimatablePropertyChanged in HandleAssetAdded. */
 	bool HandleSequenceAdded(float KeyTime, UMovieSceneSequence* Sequence);
+
+	/** Check if we can record a new sequence (deny it if one is already primed) */
+	bool CanRecordNewSequence() const;
+
+	/** Handle recording new sequence into a sub track */
+	void HandleRecordNewSequence();
+
+	/** Actually handles the adding of the section */
+	bool HandleRecordNewSequenceInternal(float KeyTime);
 };
diff --git a/Engine/Source/Editor/MovieSceneTools/Private/TrackEditors/TransformTrackEditor.cpp b/Engine/Source/Editor/MovieSceneTools/Private/TrackEditors/TransformTrackEditor.cpp
--- a/Engine/Source/Editor/MovieSceneTools/Private/TrackEditors/TransformTrackEditor.cpp
+++ b/Engine/Source/Editor/MovieSceneTools/Private/TrackEditors/TransformTrackEditor.cpp
@@ -16,6 +16,7 @@
 #include "MatineeImportTools.h"
 #include "Matinee/InterpTrackMove.h"
 #include "FloatCurveKeyArea.h"
+#include "ContentBrowserModule.h"
 
 
 #define LOCTEXT_NAMESPACE "MovieScene_TransformTrack"
@@ -257,6 +258,33 @@ void F3DTransformTrackEditor::BuildTrackContextMenu( FMenuBuilder& MenuBuilder,
 		FUIAction(
 			FExecuteAction::CreateStatic( &FMatineeImportTools::CopyInterpMoveTrack, GetSequencer().ToSharedRef(), MoveTrack, TransformTrack ),
 			FCanExecuteAction::CreateLambda( [=]()->bool { return MoveTrack != nullptr && MoveTrack->GetNumKeys() > 0 && TransformTrack != nullptr; } ) ) );
+
+	auto AnimSubMenuDelegate = [](FMenuBuilder& InMenuBuilder, TSharedRef<ISequencer> InSequencer, UMovieScene3DTransformTrack* InTransformTrack)
+	{
+		FAssetPickerConfig AssetPickerConfig;
+		AssetPickerConfig.SelectionMode = ESelectionMode::Single;
+		AssetPickerConfig.Filter.ClassNames.Add(UAnimSequence::StaticClass()->GetFName());
+		AssetPickerConfig.OnAssetSelected = FOnAssetSelected::CreateStatic(&F3DTransformTrackEditor::ImportAnimSequenceTransforms, InSequencer, InTransformTrack);
+
+		FContentBrowserModule& ContentBrowserModule = FModuleManager::LoadModuleChecked<FContentBrowserModule>(TEXT("ContentBrowser"));
+
+		InMenuBuilder.AddWidget(
+			SNew(SBox)
+			.WidthOverride(200.0f)
+			.HeightOverride(400.0f)
+			[
+				ContentBrowserModule.Get().CreateAssetPicker(AssetPickerConfig)
+			], 
+			FText(), true, false);
+	};
+
+	const bool bInOpenSubMenuOnClick = true;
+
+	MenuBuilder.AddSubMenu(
+		NSLOCTEXT( "Sequencer", "ImportTransforms", "Import From Animation Root" ),
+		NSLOCTEXT( "Sequencer", "ImportTransformsTooltip", "Import transform keys from an animation sequence's root motion." ),
+		FNewMenuDelegate::CreateLambda(AnimSubMenuDelegate, GetSequencer().ToSharedRef(), TransformTrack), 
+		bInOpenSubMenuOnClick);
 }
 
 
@@ -413,13 +441,13 @@ void F3DTransformTrackEditor::BuildObjectBindingEditButtons(TSharedPtr<SHorizont
 {
 	bool bHasCameraComponent = false;
 
-	TArray<UObject*> OutObjects;
+	TArray<TWeakObjectPtr<UObject>> OutObjects;
 	GetSequencer()->GetRuntimeObjects( GetSequencer()->GetFocusedMovieSceneSequenceInstance(), ObjectGuid, OutObjects);
 
-	for (UObject* Object : OutObjects)
+	for (auto Object : OutObjects)
 	{
-		AActor* Actor = Cast<AActor>( Object );
-		if (Actor)
+		AActor* Actor = Cast<AActor>(Object.Get());
+		if (Actor != nullptr)
 		{
 			UCameraComponent* CameraComponent = MovieSceneHelpers::CameraComponentFromActor(Actor);
 			if (CameraComponent)
@@ -488,13 +516,14 @@ bool F3DTransformTrackEditor::CanAddTransformTrackForActorHandle( FGuid ObjectBi
 EVisibility
 F3DTransformTrackEditor::IsCameraVisible(FGuid ObjectGuid) const
 {
-	TArray<UObject*> OutObjects;
+	TArray<TWeakObjectPtr<UObject>> OutObjects;
 	GetSequencer()->GetRuntimeObjects( GetSequencer()->GetFocusedMovieSceneSequenceInstance(), ObjectGuid, OutObjects);
 
-	for (UObject* Object : OutObjects)
+	for (auto Object : OutObjects)
 	{
-		AActor* Actor = Cast<AActor>( Object );
-		if (Actor)
+		AActor* Actor = Cast<AActor>(Object.Get());
+		
+		if (Actor != nullptr)
 		{
 			return EVisibility::Visible;
 		}
@@ -505,15 +534,16 @@ F3DTransformTrackEditor::IsCameraVisible(FGuid ObjectGuid) const
 
 ECheckBoxState F3DTransformTrackEditor::IsCameraLocked(FGuid ObjectGuid) const
 {
-	TArray<UObject*> OutObjects;
+	TArray<TWeakObjectPtr<UObject>> OutObjects;
 	GetSequencer()->GetRuntimeObjects( GetSequencer()->GetFocusedMovieSceneSequenceInstance(), ObjectGuid, OutObjects);
 
 	TWeakObjectPtr<AActor> CameraActor;
 
-	for (UObject* Object : OutObjects)
+	for (auto Object : OutObjects)
 	{
-		AActor* Actor = Cast<AActor>( Object );
-		if (Actor)
+		AActor* Actor = Cast<AActor>(Object.Get());
+		
+		if (Actor != nullptr)
 		{
 			CameraActor = Actor;
 			break;
@@ -522,10 +552,35 @@ ECheckBoxState F3DTransformTrackEditor::IsCameraLocked(FGuid ObjectGuid) const
 
 	if (CameraActor.IsValid())
 	{
+		// First, check the active viewport
+		FViewport* ActiveViewport = GEditor->GetActiveViewport();
+
 		for (int32 i = 0; i < GEditor->LevelViewportClients.Num(); ++i)
 		{		
 			FLevelEditorViewportClient* LevelVC = GEditor->LevelViewportClients[i];
-			if (LevelVC && LevelVC->IsPerspective() && LevelVC->AllowsCinematicPreview() && LevelVC->GetViewMode() != VMI_Unknown && CameraActor.IsValid() && LevelVC->IsActorLocked(CameraActor.Get()))
+
+			if (LevelVC && LevelVC->IsPerspective() && LevelVC->GetViewMode() != VMI_Unknown)
+			{
+				if (LevelVC->Viewport == ActiveViewport)
+				{
+					if (CameraActor.IsValid() && LevelVC->IsActorLocked(CameraActor.Get()))
+					{
+						return ECheckBoxState::Checked;
+					}
+					else
+					{
+						return ECheckBoxState::Unchecked;
+					}
+				}
+			}
+		}
+
+		// Otherwise check all other viewports
+		for (int32 i = 0; i < GEditor->LevelViewportClients.Num(); ++i)
+		{		
+			FLevelEditorViewportClient* LevelVC = GEditor->LevelViewportClients[i];
+
+			if (LevelVC && LevelVC->IsPerspective() && LevelVC->GetViewMode() != VMI_Unknown && CameraActor.IsValid() && LevelVC->IsActorLocked(CameraActor.Get()))
 			{
 				return ECheckBoxState::Checked;
 			}
@@ -538,15 +593,16 @@ ECheckBoxState F3DTransformTrackEditor::IsCameraLocked(FGuid ObjectGuid) const
 
 void F3DTransformTrackEditor::OnLockCameraClicked(ECheckBoxState CheckBoxState, FGuid ObjectGuid)
 {
-	TArray<UObject*> OutObjects;
+	TArray<TWeakObjectPtr<UObject>> OutObjects;
 	GetSequencer()->GetRuntimeObjects( GetSequencer()->GetFocusedMovieSceneSequenceInstance(), ObjectGuid, OutObjects);
 
 	TWeakObjectPtr<AActor> CameraActor;
 
-	for (UObject* Object : OutObjects)
+	for (auto Object : OutObjects)
 	{
-		AActor* Actor = Cast<AActor>( Object );
-		if (Actor)
+		AActor* Actor = Cast<AActor>(Object.Get());
+
+		if (Actor != nullptr)
 		{
 			CameraActor = Actor;
 			break;
@@ -591,7 +647,7 @@ void F3DTransformTrackEditor::OnLockCameraClicked(ECheckBoxState CheckBoxState,
 		for (int32 i = 0; i < GEditor->LevelViewportClients.Num(); ++i)
 		{		
 			FLevelEditorViewportClient* LevelVC = GEditor->LevelViewportClients[i];
-			if (LevelVC && LevelVC->IsPerspective() && LevelVC->AllowsCinematicPreview() && LevelVC->GetViewMode() != VMI_Unknown && CameraActor.IsValid() && LevelVC->IsActorLocked(CameraActor.Get()))
+			if (LevelVC && LevelVC->IsPerspective() && LevelVC->GetViewMode() != VMI_Unknown && CameraActor.IsValid() && LevelVC->IsActorLocked(CameraActor.Get()))
 			{
 				LevelVC->SetMatineeActorLock(nullptr);
 				LevelVC->SetActorLock(nullptr);
@@ -608,15 +664,16 @@ void F3DTransformTrackEditor::OnLockCameraClicked(ECheckBoxState CheckBoxState,
 
 FText F3DTransformTrackEditor::GetLockCameraToolTip(FGuid ObjectGuid) const
 {
-	TArray<UObject*> OutObjects;
+	TArray<TWeakObjectPtr<UObject>> OutObjects;
 	GetSequencer()->GetRuntimeObjects( GetSequencer()->GetFocusedMovieSceneSequenceInstance(), ObjectGuid, OutObjects);
 
 	TWeakObjectPtr<AActor> CameraActor;
 
-	for (UObject* Object : OutObjects)
+	for (auto Object : OutObjects)
 	{
-		AActor* Actor = Cast<AActor>( Object );
-		if (Actor)
+		AActor* Actor = Cast<AActor>(Object.Get());
+
+		if (Actor != nullptr)
 		{
 			CameraActor = Actor;
 			break;
@@ -682,12 +739,12 @@ void F3DTransformTrackEditor::GetTransformKeys( const FTransformData& LastTransf
 
 void F3DTransformTrackEditor::AddTransformKeysForHandle( FGuid ObjectHandle, EKey3DTransformChannel::Type ChannelToKey, FKeyParams KeyParams )
 {
-	TArray<UObject*> OutObjects;
+	TArray<TWeakObjectPtr<UObject>> OutObjects;
 	GetSequencer()->GetRuntimeObjects( GetSequencer()->GetFocusedMovieSceneSequenceInstance(), ObjectHandle, OutObjects );
 
-	for ( UObject* Object : OutObjects )
+	for ( TWeakObjectPtr<UObject> Object : OutObjects )
 	{
-		AddTransformKeysForObject(Object, ChannelToKey, KeyParams);
+		AddTransformKeysForObject(Object.Get(), ChannelToKey, KeyParams);
 	}
 }
 
@@ -750,5 +807,155 @@ void F3DTransformTrackEditor::SetIntermediateValueFromTransformChange( UMovieSce
 	OnSetIntermediateValueFromTransformChange.Broadcast(Track, TransformData);
 }
 
+void F3DTransformTrackEditor::ImportAnimSequenceTransforms(const FAssetData& Asset, TSharedRef<ISequencer> Sequencer, UMovieScene3DTransformTrack* TransformTrack)
+{
+	FSlateApplication::Get().DismissAllMenus();
+
+	UAnimSequence* AnimSequence = Cast<UAnimSequence>(Asset.GetAsset());
+
+	// find object binding to recover any component transforms we need to incorporate (for characters)
+	FTransform InvComponentTransform;
+	UMovieSceneSequence* MovieSceneSequence = Sequencer->GetFocusedMovieSceneSequence();
+	if(MovieSceneSequence)
+	{
+		UMovieScene* MovieScene = MovieSceneSequence->GetMovieScene();
+		if(MovieScene)
+		{
+			FGuid ObjectBinding;
+			if(MovieScene->FindTrackBinding(*TransformTrack, ObjectBinding))
+			{
+				const UClass* ObjectClass = nullptr;
+				if(FMovieSceneSpawnable* Spawnable = MovieScene->FindSpawnable(ObjectBinding))
+				{
+					ObjectClass = Spawnable->GetClass();
+				}
+				else if(FMovieScenePossessable* Possessable = MovieScene->FindPossessable(ObjectBinding))
+				{
+					ObjectClass = Possessable->GetPossessedObjectClass();
+				}
+
+				if(ObjectClass)
+				{
+					const ACharacter* Character = Cast<const ACharacter>(ObjectClass->ClassDefaultObject);
+					if(Character)
+					{
+						const USkeletalMeshComponent* SkeletalMeshComponent = Character->GetMesh();
+						FTransform MeshRelativeTransform = SkeletalMeshComponent->GetRelativeTransform();
+						InvComponentTransform = MeshRelativeTransform.GetRelativeTransform(SkeletalMeshComponent->GetOwner()->GetTransform()).Inverse();
+					}
+				}
+			}
+		}
+	}
+
+	if(AnimSequence && AnimSequence->RawAnimationData.Num() > 0)
+	{
+		const FScopedTransaction Transaction( NSLOCTEXT( "Sequencer", "ImportAnimSequenceTransforms", "Import Anim Sequence Transforms" ) );
+
+		TransformTrack->Modify();
+
+		UMovieScene3DTransformSection* Section = Cast<UMovieScene3DTransformSection>(TransformTrack->CreateNewSection());
+		Section->GetScaleCurve(EAxis::X).SetDefaultValue(1);
+		Section->GetScaleCurve(EAxis::Y).SetDefaultValue(1);
+		Section->GetScaleCurve(EAxis::Z).SetDefaultValue(1);
+		TransformTrack->AddSection(*Section);
+
+		if (Section->TryModify())
+		{
+			float SectionMin = Section->GetStartTime();
+			float SectionMax = Section->GetEndTime();
+
+			struct FTempTransformKey
+			{
+				FTransform Transform;
+				FRotator WoundRotation;
+				float Time;
+			};
+
+			TArray<FTempTransformKey> TempKeys;
+
+			FRawAnimSequenceTrack& RawTrack = AnimSequence->RawAnimationData[0];
+			const int32 KeyCount = FMath::Max(FMath::Max(RawTrack.PosKeys.Num(), RawTrack.RotKeys.Num()), RawTrack.ScaleKeys.Num());
+			for(int32 KeyIndex = 0; KeyIndex < KeyCount; KeyIndex++)
+			{
+				FTempTransformKey TempKey;
+				TempKey.Time = AnimSequence->GetTimeAtFrame(KeyIndex);
+
+				if(RawTrack.PosKeys.IsValidIndex(KeyIndex))
+				{
+					TempKey.Transform.SetTranslation(RawTrack.PosKeys[KeyIndex]);
+				}
+				else if(RawTrack.PosKeys.Num() > 0)
+				{
+					TempKey.Transform.SetTranslation(RawTrack.PosKeys[0]);
+				}
+				
+				if(RawTrack.RotKeys.IsValidIndex(KeyIndex))
+				{
+					TempKey.Transform.SetRotation(RawTrack.RotKeys[KeyIndex]);
+				}
+				else if(RawTrack.RotKeys.Num() > 0)
+				{
+					TempKey.Transform.SetRotation(RawTrack.RotKeys[0]);
+				}
+
+				if(RawTrack.ScaleKeys.IsValidIndex(KeyIndex))
+				{
+					TempKey.Transform.SetScale3D(RawTrack.ScaleKeys[KeyIndex]);
+				}
+				else if(RawTrack.ScaleKeys.Num() > 0)
+				{
+					TempKey.Transform.SetScale3D(RawTrack.ScaleKeys[0]);
+				}
+
+				// apply component transform if any
+				TempKey.Transform = InvComponentTransform * TempKey.Transform;
+
+				TempKey.WoundRotation = TempKey.Transform.GetRotation().Rotator();
+
+				TempKeys.Add(TempKey);
+			}
+
+			int32 TransformCount = TempKeys.Num();
+			for(int32 TransformIndex = 0; TransformIndex < TransformCount - 1; TransformIndex++)
+			{
+				FRotator& Rotator = TempKeys[TransformIndex].WoundRotation;
+				FRotator& NextRotator = TempKeys[TransformIndex + 1].WoundRotation;
+
+				FMath::WindRelativeAnglesDegrees(Rotator.Pitch, NextRotator.Pitch);
+				FMath::WindRelativeAnglesDegrees(Rotator.Yaw, NextRotator.Yaw);
+				FMath::WindRelativeAnglesDegrees(Rotator.Roll, NextRotator.Roll);
+			}
+
+			const bool bUnwindRotation = false;
+			for(const FTempTransformKey& TempKey : TempKeys)
+			{
+				SectionMin = FMath::Min( SectionMin, TempKey.Time );
+				SectionMax = FMath::Max( SectionMax, TempKey.Time );
+
+				Section->SetStartTime( SectionMin );
+				Section->SetEndTime( SectionMax );
+
+				const FVector Translation = TempKey.Transform.GetTranslation();
+				const FVector Rotation = TempKey.WoundRotation.Euler();
+				const FVector Scale = TempKey.Transform.GetScale3D();
+
+				Section->AddKey(TempKey.Time, FTransformKey(EKey3DTransformChannel::Translation, EAxis::X, Translation.X, bUnwindRotation), EMovieSceneKeyInterpolation::Linear);
+				Section->AddKey(TempKey.Time, FTransformKey(EKey3DTransformChannel::Translation, EAxis::Y, Translation.Y, bUnwindRotation), EMovieSceneKeyInterpolation::Linear);
+				Section->AddKey(TempKey.Time, FTransformKey(EKey3DTransformChannel::Translation, EAxis::Z, Translation.Z, bUnwindRotation), EMovieSceneKeyInterpolation::Linear);
+
+				Section->AddKey(TempKey.Time, FTransformKey(EKey3DTransformChannel::Rotation, EAxis::X, Rotation.X, bUnwindRotation), EMovieSceneKeyInterpolation::Linear);
+				Section->AddKey(TempKey.Time, FTransformKey(EKey3DTransformChannel::Rotation, EAxis::Y, Rotation.Y, bUnwindRotation), EMovieSceneKeyInterpolation::Linear);
+				Section->AddKey(TempKey.Time, FTransformKey(EKey3DTransformChannel::Rotation, EAxis::Z, Rotation.Z, bUnwindRotation), EMovieSceneKeyInterpolation::Linear);
+
+				Section->AddKey(TempKey.Time, FTransformKey(EKey3DTransformChannel::Scale, EAxis::X, Scale.X, bUnwindRotation), EMovieSceneKeyInterpolation::Linear);
+				Section->AddKey(TempKey.Time, FTransformKey(EKey3DTransformChannel::Scale, EAxis::Y, Scale.Y, bUnwindRotation), EMovieSceneKeyInterpolation::Linear);
+				Section->AddKey(TempKey.Time, FTransformKey(EKey3DTransformChannel::Scale, EAxis::Z, Scale.Z, bUnwindRotation), EMovieSceneKeyInterpolation::Linear);
+			}
+
+			Sequencer->NotifyMovieSceneDataChanged();
+		}
+	}
+}
 
 #undef LOCTEXT_NAMESPACE
diff --git a/Engine/Source/Editor/MovieSceneTools/Private/TrackEditors/TransformTrackEditor.h b/Engine/Source/Editor/MovieSceneTools/Private/TrackEditors/TransformTrackEditor.h
--- a/Engine/Source/Editor/MovieSceneTools/Private/TrackEditors/TransformTrackEditor.h
+++ b/Engine/Source/Editor/MovieSceneTools/Private/TrackEditors/TransformTrackEditor.h
@@ -135,6 +135,9 @@ private:
 
 	void SetIntermediateValueFromTransformChange( UMovieSceneTrack* Track, FTransformData TransformData );
 
+	/** Import an animation sequence's root transforms into a transform section */
+	static void ImportAnimSequenceTransforms(const FAssetData& Asset, TSharedRef<class ISequencer> Sequencer, UMovieScene3DTransformTrack* TransformTrack);
+
 private:
 
 	static FName TransformPropertyName;
diff --git a/Engine/Source/Editor/MovieSceneTools/Public/BoolKeyArea.h b/Engine/Source/Editor/MovieSceneTools/Public/BoolKeyArea.h
--- a/Engine/Source/Editor/MovieSceneTools/Public/BoolKeyArea.h
+++ b/Engine/Source/Editor/MovieSceneTools/Public/BoolKeyArea.h
@@ -24,4 +24,8 @@ public:
 
 	virtual bool CanCreateKeyEditor() override;
 	virtual TSharedRef<SWidget> CreateKeyEditor(ISequencer* Sequencer) override;
+
+protected:
+
+	void OnValueChanged(bool InValue);
 };
diff --git a/Engine/Source/Editor/MovieSceneTools/Public/EnumKeyArea.h b/Engine/Source/Editor/MovieSceneTools/Public/EnumKeyArea.h
--- a/Engine/Source/Editor/MovieSceneTools/Public/EnumKeyArea.h
+++ b/Engine/Source/Editor/MovieSceneTools/Public/EnumKeyArea.h
@@ -26,6 +26,10 @@ public:
 	virtual bool CanCreateKeyEditor() override;
 	virtual TSharedRef<SWidget> CreateKeyEditor(ISequencer* Sequencer) override;
 
+protected:
+
+	void OnValueChanged(int32 InValue);
+
 private:
 
 	/** The enum which provides available integral values for this key area. */
diff --git a/Engine/Source/Editor/MovieSceneTools/Public/FloatCurveKeyArea.h b/Engine/Source/Editor/MovieSceneTools/Public/FloatCurveKeyArea.h
--- a/Engine/Source/Editor/MovieSceneTools/Public/FloatCurveKeyArea.h
+++ b/Engine/Source/Editor/MovieSceneTools/Public/FloatCurveKeyArea.h
@@ -59,6 +59,10 @@ public:
 	virtual void CopyKeys(FMovieSceneClipboardBuilder& ClipboardBuilder, const TFunctionRef<bool(FKeyHandle, const IKeyArea&)>& KeyMask) const override;
 	virtual void PasteKeys(const FMovieSceneClipboardKeyTrack& KeyTrack, const FMovieSceneClipboardEnvironment& SrcEnvironment, const FSequencerPasteEnvironment& DstEnvironment) override;
 
+protected:
+
+	void OnValueChanged(float InValue);
+
 private:
 
 	/** The key area's color. */
diff --git a/Engine/Source/Editor/MovieSceneTools/Public/KeyframeTrackEditor.h b/Engine/Source/Editor/MovieSceneTools/Public/KeyframeTrackEditor.h
--- a/Engine/Source/Editor/MovieSceneTools/Public/KeyframeTrackEditor.h
+++ b/Engine/Source/Editor/MovieSceneTools/Public/KeyframeTrackEditor.h
@@ -6,6 +6,7 @@
 #include "IKeyframeSection.h"
 #include "ISequencerObjectChangeListener.h"
 #include "MovieSceneCommonHelpers.h"
+#include "MovieSceneSection.h"
 
 class IPropertyHandle;
 class FPropertyChangedParams;
@@ -151,6 +152,8 @@ private:
 		else
 		{
 			EMovieSceneKeyInterpolation InterpolationMode = GetSequencer()->GetKeyInterpolation();
+			bool bInfiniteKeyAreas = GetSequencer()->GetInfiniteKeyAreas();
+
 			for ( const KeyDataType& NewKey : NewKeys )
 			{
 				// Only modify the track if the new key will create new data, or if creating keys has been forced.
@@ -158,11 +161,11 @@ private:
 				{
 					if ( HasKeys( Track, NewKey ) || KeyParams.bCreateKeyIfEmpty )
 					{
-						bSectionCreated |= AddKey( Track, KeyTime, NewKey, InterpolationMode );
+						bSectionCreated |= AddKey( Track, KeyTime, NewKey, InterpolationMode, bInfiniteKeyAreas );
 					}
 					else
 					{
-						bSectionCreated |= SetDefault( Track, KeyTime, NewKey );
+						bSectionCreated |= SetDefault( Track, KeyTime, NewKey, bInfiniteKeyAreas );
 					}
 				}
 			}
@@ -171,7 +174,7 @@ private:
 			{
 				for ( const KeyDataType& DefaultKey : DefaultKeys )
 				{
-					bSectionCreated |= SetDefault( Track, KeyTime, DefaultKey );
+					bSectionCreated |= SetDefault( Track, KeyTime, DefaultKey, bInfiniteKeyAreas  );
 				}
 			}
 		}
@@ -200,17 +203,23 @@ private:
 	using FMovieSceneTrackEditor::AddKey;
 
 	/* Return whether a section was added */
-	bool AddKey( TrackType* Track, float Time, const KeyDataType& KeyData, EMovieSceneKeyInterpolation KeyInterpolation )
+	bool AddKey( TrackType* Track, float Time, const KeyDataType& KeyData, EMovieSceneKeyInterpolation KeyInterpolation, bool bInfiniteKeyAreas )
 	{
 		bool bSectionAdded = false;
 		Track->Modify();
-		IKeyframeSection<KeyDataType>* KeyframeSection = CastChecked<SectionType>( Track->FindOrAddSection( Time, bSectionAdded ) );
+		UMovieSceneSection* NewSection = Track->FindOrAddSection( Time, bSectionAdded );
+		IKeyframeSection<KeyDataType>* KeyframeSection = CastChecked<SectionType>( NewSection );
 		KeyframeSection->AddKey( Time, KeyData, KeyInterpolation );
+
+		if (bSectionAdded)
+		{
+			NewSection->SetIsInfinite(bInfiniteKeyAreas);
+		}
 		return bSectionAdded;
 	}
 
 	/* Return whether a section was added */
-	bool SetDefault( TrackType* Track, float Time, const KeyDataType& KeyData )
+	bool SetDefault( TrackType* Track, float Time, const KeyDataType& KeyData, bool bInfiniteKeyAreas )
 	{
 		bool bSectionAdded = false;
 		Track->Modify();
@@ -225,8 +234,14 @@ private:
 		}
 		else
 		{
-			IKeyframeSection<KeyDataType>* KeyframeSection = CastChecked<SectionType>( Track->FindOrAddSection( Time, bSectionAdded ) );
+			UMovieSceneSection* NewSection = Track->FindOrAddSection( Time, bSectionAdded );
+			IKeyframeSection<KeyDataType>* KeyframeSection = CastChecked<SectionType>( NewSection );
 			KeyframeSection->SetDefault( KeyData );
+
+			if (bSectionAdded)
+			{
+				NewSection->SetIsInfinite(bInfiniteKeyAreas);
+			}
 		}
 		return bSectionAdded;
 	}
diff --git a/Engine/Source/Editor/MovieSceneTools/Public/MovieSceneToolHelpers.h b/Engine/Source/Editor/MovieSceneTools/Public/MovieSceneToolHelpers.h
--- a/Engine/Source/Editor/MovieSceneTools/Public/MovieSceneToolHelpers.h
+++ b/Engine/Source/Editor/MovieSceneTools/Public/MovieSceneToolHelpers.h
@@ -4,6 +4,7 @@
 
 class UMovieSceneSection;
 
+DECLARE_DELEGATE_TwoParams(FOnEnumSelectionChanged, int32 /*Selection*/, ESelectInfo::Type /*SelectionType*/);
 
 class MOVIESCENETOOLS_API MovieSceneToolHelpers
 {
@@ -57,7 +58,7 @@ public:
 	static FString GenerateNewShotName(const TArray<UMovieSceneSection*>& AllSections, float Time);
 
 	/**
-	 * Gather takes (level sequence assets that have the same shot prefix and shot number)
+	 * Gather takes - level sequence assets that have the same shot prefix and shot number in the same asset path (directory)
 	 * 
 	 * @param Section The section to gather takes from
 	 * @param TakeNumbers The gathered take numbers
@@ -74,4 +75,14 @@ public:
 	 * @return The asset
 	 */
 	static UObject* GetTake(const UMovieSceneSection* Section, uint32 TakeNumber);
+
+	/**
+	 * Generate a combobox for editing enum values
+	 *
+	 * @param Enum The enum to make the combobox from
+	 * @param CurrentValue The current value to display
+	 * @param OnSelectionChanged Delegate fired when selection is changed
+	 * @return The new widget
+	 */
+	static TSharedRef<SWidget> MakeEnumComboBox(const UEnum* Enum, TAttribute<int32> CurrentValue, FOnEnumSelectionChanged OnSelectionChanged, TAttribute<TOptional<uint8>> InIntermediateValue);
 };
diff --git a/Engine/Source/Editor/MovieSceneTools/Public/MovieSceneToolsProjectSettings.h b/Engine/Source/Editor/MovieSceneTools/Public/MovieSceneToolsProjectSettings.h
--- a/Engine/Source/Editor/MovieSceneTools/Public/MovieSceneToolsProjectSettings.h
+++ b/Engine/Source/Editor/MovieSceneTools/Public/MovieSceneToolsProjectSettings.h
@@ -21,6 +21,10 @@ public:
 	UPROPERTY(config, EditAnywhere, Category=Timeline, meta=(ClampMin=0.00001f, Units=s))
 	float DefaultDuration;
 
+	/** The default directory for the shots. */
+	UPROPERTY(config, EditAnywhere, Category=Shots)
+	FString ShotDirectory;
+
 	/** The default prefix for shot names. */
 	UPROPERTY(config, EditAnywhere, Category=Shots)
 	FString ShotPrefix;
diff --git a/Engine/Source/Editor/NiagaraEditor/Private/NiagaraSequencer.cpp b/Engine/Source/Editor/NiagaraEditor/Private/NiagaraSequencer.cpp
--- a/Engine/Source/Editor/NiagaraEditor/Private/NiagaraSequencer.cpp
+++ b/Engine/Source/Editor/NiagaraEditor/Private/NiagaraSequencer.cpp
@@ -20,7 +20,7 @@ UNiagaraMovieSceneSection::UNiagaraMovieSceneSection(const FObjectInitializer& O
 /*
  *  This is called when the user edits a track in the effect editor timeline
  */
-void INiagaraTrackInstance::RefreshInstance(const TArray<UObject*>& RuntimeObjects, IMovieScenePlayer& Player, FMovieSceneSequenceInstance& SequenceInstance)
+void INiagaraTrackInstance::RefreshInstance(const TArray<TWeakObjectPtr<UObject>>& RuntimeObjects, IMovieScenePlayer& Player, FMovieSceneSequenceInstance& SequenceInstance)
 {
 	// every track should have exactly one section, always
 	check(Track->GetAllSections().Num() == 1)
diff --git a/Engine/Source/Editor/NiagaraEditor/Private/NiagaraSequencer.h b/Engine/Source/Editor/NiagaraEditor/Private/NiagaraSequencer.h
--- a/Engine/Source/Editor/NiagaraEditor/Private/NiagaraSequencer.h
+++ b/Engine/Source/Editor/NiagaraEditor/Private/NiagaraSequencer.h
@@ -120,15 +120,14 @@ public:
 	{
 	}
 
-	virtual void Update(EMovieSceneUpdateData& UpdateData, const TArray<UObject*>& RuntimeObjects, IMovieScenePlayer& Player, FMovieSceneSequenceInstance& SequenceInstance)
+	virtual void Update(EMovieSceneUpdateData& UpdateData, const TArray<TWeakObjectPtr<UObject>>& RuntimeObjects, IMovieScenePlayer& Player, FMovieSceneSequenceInstance& SequenceInstance)
 	{
 	}
 
-	virtual void RefreshInstance(const TArray<UObject*>& RuntimeObjects, IMovieScenePlayer& Player, FMovieSceneSequenceInstance& SequenceInstance);
-
-	virtual void ClearInstance( IMovieScenePlayer& Player, FMovieSceneSequenceInstance& SequenceInstance ) override {}
-	virtual void SaveState(const TArray<UObject*>& RuntimeObjects, IMovieScenePlayer& Player, FMovieSceneSequenceInstance& SequenceInstance) override {}
-	virtual void RestoreState(const TArray<UObject*>& RuntimeObjects, IMovieScenePlayer& Player, FMovieSceneSequenceInstance& SequenceInstance) override {}
+	virtual void RefreshInstance(const TArray<TWeakObjectPtr<UObject>>& RuntimeObjects, IMovieScenePlayer& Player, FMovieSceneSequenceInstance& SequenceInstance);
+	virtual void ClearInstance(IMovieScenePlayer& Player, FMovieSceneSequenceInstance& SequenceInstance) override {}
+	virtual void SaveState(const TArray<TWeakObjectPtr<UObject>>& RuntimeObjects, IMovieScenePlayer& Player, FMovieSceneSequenceInstance& SequenceInstance) override {}
+	virtual void RestoreState(const TArray<TWeakObjectPtr<UObject>>& RuntimeObjects, IMovieScenePlayer& Player, FMovieSceneSequenceInstance& SequenceInstance) override {}
 
 private:
 	TSharedPtr<FNiagaraSimulation> Emitter;
diff --git a/Engine/Source/Editor/Persona/Private/Persona.cpp b/Engine/Source/Editor/Persona/Private/Persona.cpp
--- a/Engine/Source/Editor/Persona/Private/Persona.cpp
+++ b/Engine/Source/Editor/Persona/Private/Persona.cpp
@@ -3045,13 +3045,6 @@ void FPersona::Tick(float DeltaTime)
 		Viewport.Pin()->RefreshViewport();
 	}
 
-	if (IsRecording())
-	{
-		// make sure you don't allow switch previewcomponent
-		FPersonaModule& PersonaModule = FModuleManager::GetModuleChecked<FPersonaModule>("Persona");
-		PersonaModule.OnTickRecording().ExecuteIfBound(PreviewComponent, DeltaTime);
-	}
-
 	if (PreviewComponent && LastCachedLODForPreviewComponent != PreviewComponent->PredictedLODLevel)
 	{
 		OnLODChanged.Broadcast();
diff --git a/Engine/Source/Editor/Persona/Public/PersonaModule.h b/Engine/Source/Editor/Persona/Public/PersonaModule.h
--- a/Engine/Source/Editor/Persona/Public/PersonaModule.h
+++ b/Engine/Source/Editor/Persona/Public/PersonaModule.h
@@ -72,9 +72,6 @@ public:
 	/** Delegate used to get the currently recording animation time */
 	FGetCurrentRecordingTime& OnGetCurrentRecordingTime() { return GetCurrentRecordingTimeDelegate; }
 
-	/** Delegate used to tick the skelmesh component recording */
-	FTickRecording& OnTickRecording() { return TickRecordingDelegate; }
-
 private:
 	/** When a new AnimBlueprint is created, this will handle post creation work such as adding non-event default nodes */
 	void OnNewBlueprintCreated(UBlueprint* InBlueprint);
diff --git a/Engine/Source/Editor/PhAT/Private/PhATPreviewViewportClient.cpp b/Engine/Source/Editor/PhAT/Private/PhATPreviewViewportClient.cpp
--- a/Engine/Source/Editor/PhAT/Private/PhATPreviewViewportClient.cpp
+++ b/Engine/Source/Editor/PhAT/Private/PhATPreviewViewportClient.cpp
@@ -684,12 +684,6 @@ void FPhATEdPreviewViewportClient::Tick(float DeltaSeconds)
 	}
 
 	World->Tick(LEVELTICK_All, DeltaSeconds * SharedData->EditorSimOptions->TimeDilation);
-
-	if (PhATPtr.Pin()->IsRecording())
-	{
-		FPersonaModule& PersonaModule = FModuleManager::GetModuleChecked<FPersonaModule>("Persona");
-		PersonaModule.OnTickRecording().ExecuteIfBound(SharedData->EditorSkelComp, DeltaSeconds);
-	}
 }
 
 FSceneInterface* FPhATEdPreviewViewportClient::GetScene() const
diff --git a/Engine/Source/Editor/PropertyEditor/Private/DetailLayoutBuilderImpl.cpp b/Engine/Source/Editor/PropertyEditor/Private/DetailLayoutBuilderImpl.cpp
--- a/Engine/Source/Editor/PropertyEditor/Private/DetailLayoutBuilderImpl.cpp
+++ b/Engine/Source/Editor/PropertyEditor/Private/DetailLayoutBuilderImpl.cpp
@@ -204,7 +204,7 @@ void FDetailLayoutBuilderImpl::FilterDetailLayout( const FDetailFilter& InFilter
 	}
 }
 
-void FDetailLayoutBuilderImpl::SetCurrentCustomizationClass( UClass* CurrentClass, FName VariableName )
+void FDetailLayoutBuilderImpl::SetCurrentCustomizationClass( UStruct* CurrentClass, FName VariableName )
 {
 	CurrentCustomizationClass = CurrentClass;
 	CurrentCustomizationVariableName = VariableName;
diff --git a/Engine/Source/Editor/PropertyEditor/Private/DetailLayoutBuilderImpl.h b/Engine/Source/Editor/PropertyEditor/Private/DetailLayoutBuilderImpl.h
--- a/Engine/Source/Editor/PropertyEditor/Private/DetailLayoutBuilderImpl.h
+++ b/Engine/Source/Editor/PropertyEditor/Private/DetailLayoutBuilderImpl.h
@@ -46,7 +46,7 @@ public:
 	 * @param CurrentClass	The current class being customized
 	 * @param VariableName The variable name of the class being customized (used for inner classes/structs where there can be multiple instances of them)
 	 */
-	void SetCurrentCustomizationClass( UClass* CurrentClass, FName VariableName );
+	void SetCurrentCustomizationClass( UStruct* CurrentClass, FName VariableName );
 
 	/** @return The current class variable name being customized */
 	FName GetCurrentCustomizationVariableName() const { return CurrentCustomizationVariableName; }
@@ -182,5 +182,5 @@ private:
 	/** The view where this detail customizer resides */
 	class IDetailsViewPrivate& DetailsView;
 	/** The current class being customized */
-	UClass* CurrentCustomizationClass;
+	UStruct* CurrentCustomizationClass;
 };
diff --git a/Engine/Source/Editor/PropertyEditor/Private/PropertyHandleImpl.cpp b/Engine/Source/Editor/PropertyEditor/Private/PropertyHandleImpl.cpp
--- a/Engine/Source/Editor/PropertyEditor/Private/PropertyHandleImpl.cpp
+++ b/Engine/Source/Editor/PropertyEditor/Private/PropertyHandleImpl.cpp
@@ -238,13 +238,15 @@ FString FPropertyValueImpl::GetPropertyValueArray() const
 
 bool FPropertyValueImpl::SendTextToObjectProperty( const FString& Text, EPropertyValueSetFlags::Type Flags )
 {
-	if (IsAnyOuterObjectFromEngine() && !FPackageName::IsEnginePackageName(Text))
+	TSharedPtr<FPropertyNode> PropertyNodePin = PropertyNode.Pin();
+	bool bIsLazyObjectProperty = PropertyNodePin.IsValid() && PropertyNodePin->GetProperty() && PropertyNodePin->GetProperty()->GetClass() == ULazyObjectProperty::StaticClass();
+
+	if (IsAnyOuterObjectFromEngine() && !FPackageName::IsEnginePackageName(Text) && !bIsLazyObjectProperty)
 	{
 		UE_LOG(LogPropertyNode, Log, TEXT("Cannot assign a Project object %s to an Engine property."), *Text);
 		return false;
 	}
-
-	TSharedPtr<FPropertyNode> PropertyNodePin = PropertyNode.Pin();
+	
 	if( PropertyNodePin.IsValid() )
 	{
 		FComplexPropertyNode* ParentNode = PropertyNodePin->FindComplexParent();
diff --git a/Engine/Source/Editor/PropertyEditor/Private/SDetailsView.cpp b/Engine/Source/Editor/PropertyEditor/Private/SDetailsView.cpp
--- a/Engine/Source/Editor/PropertyEditor/Private/SDetailsView.cpp
+++ b/Engine/Source/Editor/PropertyEditor/Private/SDetailsView.cpp
@@ -676,23 +676,14 @@ UStruct* SDetailsView::GetBaseStruct() const
 	return RootPropertyNode.IsValid() ? RootPropertyNode->GetObjectBaseClass() : NULL;
 }
 
-void SDetailsView::RegisterInstancedCustomPropertyLayout( UClass* Class, FOnGetDetailCustomizationInstance DetailLayoutDelegate )
+void SDetailsView::RegisterInstancedCustomPropertyLayout( UStruct* Class, FOnGetDetailCustomizationInstance DetailLayoutDelegate )
 {
-	check( Class );
-
-	FDetailLayoutCallback Callback;
-	Callback.DetailLayoutDelegate = DetailLayoutDelegate;
-	// @todo: DetailsView: Fix me: this specifies the order in which detail layouts should be queried
-	Callback.Order = InstancedClassToDetailLayoutMap.Num();
-
-	InstancedClassToDetailLayoutMap.Add( Class, Callback );
+	RegisterInstancedCustomPropertyLayoutInternal(Class, DetailLayoutDelegate);
 }
 
-void SDetailsView::UnregisterInstancedCustomPropertyLayout( UClass* Class )
+void SDetailsView::UnregisterInstancedCustomPropertyLayout( UStruct* Class )
 {
-	check( Class );
-
-	InstancedClassToDetailLayoutMap.Remove( Class );
+	UnregisterInstancedCustomPropertyLayoutInternal(Class);
 }
 
 void SDetailsView::AddExternalRootPropertyNode( TSharedRef<FPropertyNode> ExternalRootNode )
diff --git a/Engine/Source/Editor/PropertyEditor/Private/SDetailsView.h b/Engine/Source/Editor/PropertyEditor/Private/SDetailsView.h
--- a/Engine/Source/Editor/PropertyEditor/Private/SDetailsView.h
+++ b/Engine/Source/Editor/PropertyEditor/Private/SDetailsView.h
@@ -101,8 +101,8 @@ public:
 		return false; 
 	}
 private:
-	void RegisterInstancedCustomPropertyLayout( UClass* Class, FOnGetDetailCustomizationInstance DetailLayoutDelegate ) override;
-	void UnregisterInstancedCustomPropertyLayout( UClass* Class ) override;
+	void RegisterInstancedCustomPropertyLayout( UStruct* Class, FOnGetDetailCustomizationInstance DetailLayoutDelegate ) override;
+	void UnregisterInstancedCustomPropertyLayout( UStruct* Class ) override;
 	void SetObjectArrayPrivate( const TArray< TWeakObjectPtr< UObject > >& InObjects );
 
 	TSharedRef<SDetailTree> ConstructTreeView( TSharedRef<SScrollBar>& ScrollBar );
diff --git a/Engine/Source/Editor/PropertyEditor/Private/SDetailsViewBase.cpp b/Engine/Source/Editor/PropertyEditor/Private/SDetailsViewBase.cpp
--- a/Engine/Source/Editor/PropertyEditor/Private/SDetailsViewBase.cpp
+++ b/Engine/Source/Editor/PropertyEditor/Private/SDetailsViewBase.cpp
@@ -594,7 +594,7 @@ void SDetailsViewBase::QueryCustomDetailLayout(FDetailLayoutBuilderImpl& CustomD
 				for (FClassInstanceToPropertyMap::TIterator InstanceIt(InstancedPropertyMap); InstanceIt; ++InstanceIt)
 				{
 					FName Key = InstanceIt.Key();
-					CustomDetailLayout.SetCurrentCustomizationClass(CastChecked<UClass>(Class), Key);
+					CustomDetailLayout.SetCurrentCustomizationClass(Class, Key);
 
 					const FOnGetDetailCustomizationInstance& DetailDelegate = LayoutIt.Value()->DetailLayoutDelegate;
 
@@ -1206,3 +1206,21 @@ void SDetailsViewBase::UpdatePropertyMap()
 	DetailLayout->GenerateDetailLayout();
 }
 
+void SDetailsViewBase::RegisterInstancedCustomPropertyLayoutInternal(UStruct* Class, FOnGetDetailCustomizationInstance DetailLayoutDelegate)
+{
+	check( Class );
+
+	FDetailLayoutCallback Callback;
+	Callback.DetailLayoutDelegate = DetailLayoutDelegate;
+	// @todo: DetailsView: Fix me: this specifies the order in which detail layouts should be queried
+	Callback.Order = InstancedClassToDetailLayoutMap.Num();
+
+	InstancedClassToDetailLayoutMap.Add( Class, Callback );	
+}
+
+void SDetailsViewBase::UnregisterInstancedCustomPropertyLayoutInternal(UStruct* Class)
+{
+	check( Class );
+
+	InstancedClassToDetailLayoutMap.Remove( Class );	
+}
diff --git a/Engine/Source/Editor/PropertyEditor/Private/SDetailsViewBase.h b/Engine/Source/Editor/PropertyEditor/Private/SDetailsViewBase.h
--- a/Engine/Source/Editor/PropertyEditor/Private/SDetailsViewBase.h
+++ b/Engine/Source/Editor/PropertyEditor/Private/SDetailsViewBase.h
@@ -401,6 +401,13 @@ protected:
 
 	/** Called to get the visibility of the filter box */
 	EVisibility GetFilterBoxVisibility() const;
+
+	/** Utility function allowing derived classes to optionally implement RegisterInstancedCustomPropertyLayout */
+	void RegisterInstancedCustomPropertyLayoutInternal(UStruct* Class, FOnGetDetailCustomizationInstance DetailLayoutDelegate);
+
+	/** Utility function allowing derived classes to optionally implement UnregisterInstancedCustomPropertyLayout */
+	void UnregisterInstancedCustomPropertyLayoutInternal(UStruct* Class);
+
 protected:
 	/** The user defined args for the details view */
 	FDetailsViewArgs DetailsViewArgs;
diff --git a/Engine/Source/Editor/PropertyEditor/Private/SStructureDetailsView.cpp b/Engine/Source/Editor/PropertyEditor/Private/SStructureDetailsView.cpp
--- a/Engine/Source/Editor/PropertyEditor/Private/SStructureDetailsView.cpp
+++ b/Engine/Source/Editor/PropertyEditor/Private/SStructureDetailsView.cpp
@@ -258,4 +258,14 @@ EVisibility SStructureDetailsView::GetPropertyEditingVisibility() const
 	return StructData.IsValid() && StructData->IsValid() && RootNode.IsValid() && RootNode->IsValid() ? EVisibility::Visible : EVisibility::Collapsed;
 }
 
+void SStructureDetailsView::RegisterInstancedCustomPropertyLayout(UStruct* Class, FOnGetDetailCustomizationInstance DetailLayoutDelegate)
+{
+	RegisterInstancedCustomPropertyLayoutInternal(Class, DetailLayoutDelegate);
+}
+
+void SStructureDetailsView::UnregisterInstancedCustomPropertyLayout(UStruct* Class)
+{
+	UnregisterInstancedCustomPropertyLayoutInternal(Class);
+}
+
 #undef LOCTEXT_NAMESPACE
diff --git a/Engine/Source/Editor/PropertyEditor/Private/SStructureDetailsView.h b/Engine/Source/Editor/PropertyEditor/Private/SStructureDetailsView.h
--- a/Engine/Source/Editor/PropertyEditor/Private/SStructureDetailsView.h
+++ b/Engine/Source/Editor/PropertyEditor/Private/SStructureDetailsView.h
@@ -98,8 +98,8 @@ public:
 	}
 
 	virtual void SetOnObjectArrayChanged(FOnObjectArrayChanged OnObjectArrayChangedDelegate) override {}
-	virtual void RegisterInstancedCustomPropertyLayout(UClass* Class, FOnGetDetailCustomizationInstance DetailLayoutDelegate) override {}
-	virtual void UnregisterInstancedCustomPropertyLayout(UClass* Class) override {}
+	virtual void RegisterInstancedCustomPropertyLayout(UStruct* Class, FOnGetDetailCustomizationInstance DetailLayoutDelegate) override;
+	virtual void UnregisterInstancedCustomPropertyLayout(UStruct* Class) override;
 	virtual void SetObjects(const TArray<UObject*>& InObjects, bool bForceRefresh = false, bool bOverrideLock = false) override {}
 	virtual void SetObjects(const TArray< TWeakObjectPtr< UObject > >& InObjects, bool bForceRefresh = false, bool bOverrideLock = false) override {}
 	virtual void SetObject(UObject* InObject, bool bForceRefresh = false) override{}
diff --git a/Engine/Source/Editor/PropertyEditor/Private/UserInterface/PropertyEditor/PropertyEditorAssetConstants.h b/Engine/Source/Editor/PropertyEditor/Private/UserInterface/PropertyEditor/PropertyEditorAssetConstants.h
--- a/Engine/Source/Editor/PropertyEditor/Private/UserInterface/PropertyEditor/PropertyEditorAssetConstants.h
+++ b/Engine/Source/Editor/PropertyEditor/Private/UserInterface/PropertyEditor/PropertyEditorAssetConstants.h
@@ -4,5 +4,5 @@
 namespace PropertyEditorAssetConstants
 {
 	static const FVector2D ContentBrowserWindowSize(300.0f, 300.0f);
-	static const FVector2D SceneOutlinerWindowSize(300.0f, 300.0f);
+	static const FVector2D SceneOutlinerWindowSize(350.0f, 300.0f);
 };
\ No newline at end of file
diff --git a/Engine/Source/Editor/PropertyEditor/Private/UserInterface/PropertyEditor/SPropertyEditorAsset.cpp b/Engine/Source/Editor/PropertyEditor/Private/UserInterface/PropertyEditor/SPropertyEditorAsset.cpp
--- a/Engine/Source/Editor/PropertyEditor/Private/UserInterface/PropertyEditor/SPropertyEditorAsset.cpp
+++ b/Engine/Source/Editor/PropertyEditor/Private/UserInterface/PropertyEditor/SPropertyEditorAsset.cpp
@@ -568,7 +568,7 @@ bool SPropertyEditorAsset::CanAssetBeAssigned(const FAssetData& AssetData) const
 	// Prevent Engine assets being assigned non-Engine asset properties
 	const bool bHasEngineOuterObject = HasEngineOuterObject();
 	const FString AssetPath = AssetData.ObjectPath.ToString();
-	if (bHasEngineOuterObject && !FPackageName::IsEnginePackageName(AssetPath))
+	if (bHasEngineOuterObject && !FPackageName::IsEnginePackageName(AssetPath) && PropertyEditor->GetPropertyHandle()->GetPropertyClass() != ULazyObjectProperty::StaticClass())
 	{
 		FMessageDialog::Open(EAppMsgType::Ok, FText::Format(
 			LOCTEXT("ObjectAssignmentToEngineFailed", "Cannot assign a Project object {0} to an Engine property."),
diff --git a/Engine/Source/Editor/PropertyEditor/Public/IDetailsView.h b/Engine/Source/Editor/PropertyEditor/Public/IDetailsView.h
--- a/Engine/Source/Editor/PropertyEditor/Public/IDetailsView.h
+++ b/Engine/Source/Editor/PropertyEditor/Public/IDetailsView.h
@@ -127,14 +127,14 @@ public:
 	 * @param Class	The class the custom detail layout is for
 	 * @param DetailLayoutDelegate	The delegate to call when querying for custom detail layouts for the classes properties
 	 */
-	virtual void RegisterInstancedCustomPropertyLayout( UClass* Class, FOnGetDetailCustomizationInstance DetailLayoutDelegate ) = 0;
+	virtual void RegisterInstancedCustomPropertyLayout( UStruct* Class, FOnGetDetailCustomizationInstance DetailLayoutDelegate ) = 0;
 
 	/**
 	 * Unregisters a custom detail layout delegate for a specific class in this instance of the details view only
 	 *
 	 * @param Class	The class with the custom detail layout delegate to remove
 	 */
-	virtual void UnregisterInstancedCustomPropertyLayout( UClass* Class ) = 0;
+	virtual void UnregisterInstancedCustomPropertyLayout( UStruct* Class ) = 0;
 
 	/**
 	 * Sets the objects this details view is viewing
diff --git a/Engine/Source/Editor/SceneOutliner/Private/SSceneOutliner.cpp b/Engine/Source/Editor/SceneOutliner/Private/SSceneOutliner.cpp
--- a/Engine/Source/Editor/SceneOutliner/Private/SSceneOutliner.cpp
+++ b/Engine/Source/Editor/SceneOutliner/Private/SSceneOutliner.cpp
@@ -35,6 +35,56 @@ DEFINE_LOG_CATEGORY_STATIC(LogSceneOutliner, Log, All);
 
 namespace SceneOutliner
 {
+	FText GetWorldDescription(UWorld* World)
+	{
+		FText Description;
+		if(World)
+		{
+			FText PostFix;
+			const FWorldContext* WorldContext = nullptr;
+			for (const FWorldContext& Context : GEngine->GetWorldContexts())
+			{
+				if(Context.World() == World)
+				{
+					WorldContext = &Context;
+					break;
+				}
+			}
+
+			if (World->WorldType == EWorldType::PIE)
+			{
+				switch(World->GetNetMode())
+				{
+					case NM_Client:
+						if(WorldContext)
+						{
+							PostFix = FText::Format(LOCTEXT("ClientPostfixFormat", "(Client {0})"), FText::AsNumber(WorldContext->PIEInstance - 1));
+						}
+						else
+						{
+							PostFix = LOCTEXT("ClientPostfixFormat", "(Client)");
+						}
+						break;
+					case NM_DedicatedServer:
+					case NM_ListenServer:
+						PostFix = LOCTEXT("ServerPostfix", "(Server)");
+						break;
+					case NM_Standalone:
+						PostFix = LOCTEXT("PlayInEditorPostfix", "(Play In Editor)");
+						break;
+				}
+			}
+			else if(World->WorldType == EWorldType::Editor)
+			{
+				PostFix = LOCTEXT("EditorPostfix", "(Editor)");
+			}
+
+			Description = FText::Format(LOCTEXT("WorldFormat", "{0} {1}"), FText::FromString(World->GetFName().GetPlainNameString()), PostFix);	
+		}
+
+		return Description;
+	}
+
 	TSharedPtr< FOutlinerFilter > CreateSelectedActorFilter()
 	{
 		auto* Filter = new FOutlinerPredicateFilter(FActorFilterPredicate::CreateStatic([](const AActor* InActor){	return InActor->IsSelected(); }), EDefaultFilterBehaviour::Fail);
@@ -283,16 +333,16 @@ namespace SceneOutliner
 			]
 		];
 
+		// Separator
+		VerticalBox->AddSlot()
+		.AutoHeight()
+		.Padding(0, 0, 0, 1)
+		[
+			SNew(SSeparator)
+		];
+
 		if (SharedData->Mode == ESceneOutlinerMode::ActorBrowsing)
 		{
-			// Separator
-			VerticalBox->AddSlot()
-			.AutoHeight()
-			.Padding(0, 0, 0, 1)
-			[
-				SNew(SSeparator)
-			];
-
 			// Bottom panel
 			VerticalBox->AddSlot()
 			.AutoHeight()
@@ -318,7 +368,7 @@ namespace SceneOutliner
 					.ContentPadding(0)
 					.ForegroundColor( this, &SSceneOutliner::GetViewButtonForegroundColor )
 					.ButtonStyle( FEditorStyle::Get(), "ToggleButton" ) // Use the tool bar item style for this button
-					.OnGetMenuContent( this, &SSceneOutliner::GetViewButtonContent )
+					.OnGetMenuContent( this, &SSceneOutliner::GetViewButtonContent, false )
 					.ButtonContent()
 					[
 						SNew(SHorizontalBox)
@@ -341,6 +391,46 @@ namespace SceneOutliner
 				]
 			];
 		} //end if (SharedData->Mode == ESceneOutlinerMode::ActorBrowsing)
+		else
+		{
+			// Bottom panel
+			VerticalBox->AddSlot()
+			.AutoHeight()
+			[
+				SNew(SHorizontalBox)
+
+				// World picker combo button
+				+SHorizontalBox::Slot()
+				.FillWidth(1.0f)
+				.HAlign(HAlign_Right)
+				[
+					SAssignNew( ViewOptionsComboButton, SComboButton )
+					.ContentPadding(0)
+					.ForegroundColor( this, &SSceneOutliner::GetViewButtonForegroundColor )
+					.ButtonStyle( FEditorStyle::Get(), "ToggleButton" ) // Use the tool bar item style for this button
+					.OnGetMenuContent( this, &SSceneOutliner::GetViewButtonContent, true )
+					.ButtonContent()
+					[
+						SNew(SHorizontalBox)
+
+						+SHorizontalBox::Slot()
+						.AutoWidth()
+						.VAlign(VAlign_Center)
+						[
+							SNew(SImage).Image( FEditorStyle::GetBrush("SceneOutliner.World") )
+						]
+
+						+SHorizontalBox::Slot()
+						.AutoWidth()
+						.Padding(2, 0, 0, 0)
+						.VAlign(VAlign_Center)
+						[
+							SNew(STextBlock).Text( LOCTEXT("ChooseWorldMenu", "Choose World") )
+						]
+					]
+				]
+			];
+		}
 
 
 		// Don't allow tool-tips over the header
@@ -491,57 +581,113 @@ namespace SceneOutliner
 		return ViewOptionsComboButton->IsHovered() ? FEditorStyle::GetSlateColor(InvertedForegroundName) : FEditorStyle::GetSlateColor(DefaultForegroundName);
 	}
 
-	TSharedRef<SWidget> SSceneOutliner::GetViewButtonContent()
+	TSharedRef<SWidget> SSceneOutliner::GetViewButtonContent(bool bWorldPickerOnly)
 	{
-		FMenuBuilder MenuBuilder(true, NULL);
+		FMenuBuilder MenuBuilder(!bWorldPickerOnly, NULL);
 
-		MenuBuilder.BeginSection("AssetThumbnails", LOCTEXT("ShowHeading", "Show"));
+		if(bWorldPickerOnly)
 		{
-			MenuBuilder.AddMenuEntry(
-				LOCTEXT("ToggleShowOnlySelected", "Only Selected"),
-				LOCTEXT("ToggleShowOnlySelectedToolTip", "When enabled, only displays actors that are currently selected."),
-				FSlateIcon(),
-				FUIAction(
-				FExecuteAction::CreateSP( this, &SSceneOutliner::ToggleShowOnlySelected ),
-				FCanExecuteAction(),
-				FIsActionChecked::CreateSP( this, &SSceneOutliner::IsShowingOnlySelected )
-				),
-				NAME_None,
-				EUserInterfaceActionType::ToggleButton
-			);
+			BuildWorldPickerContent(MenuBuilder);
+		}
+		else
+		{
+			MenuBuilder.BeginSection("AssetThumbnails", LOCTEXT("ShowHeading", "Show"));
+			{
+				MenuBuilder.AddMenuEntry(
+					LOCTEXT("ToggleShowOnlySelected", "Only Selected"),
+					LOCTEXT("ToggleShowOnlySelectedToolTip", "When enabled, only displays actors that are currently selected."),
+					FSlateIcon(),
+					FUIAction(
+					FExecuteAction::CreateSP( this, &SSceneOutliner::ToggleShowOnlySelected ),
+					FCanExecuteAction(),
+					FIsActionChecked::CreateSP( this, &SSceneOutliner::IsShowingOnlySelected )
+					),
+					NAME_None,
+					EUserInterfaceActionType::ToggleButton
+				);
 
-			MenuBuilder.AddMenuEntry(
-				LOCTEXT("ToggleHideTemporaryActors", "Hide Temporary Actors"),
-				LOCTEXT("ToggleHideTemporaryActorsToolTip", "When enabled, hides temporary/run-time Actors."),
-				FSlateIcon(),
-				FUIAction(
-				FExecuteAction::CreateSP( this, &SSceneOutliner::ToggleHideTemporaryActors ),
-				FCanExecuteAction(),
-				FIsActionChecked::CreateSP( this, &SSceneOutliner::IsHidingTemporaryActors )
-				),
-				NAME_None,
-				EUserInterfaceActionType::ToggleButton
-			);
+				MenuBuilder.AddMenuEntry(
+					LOCTEXT("ToggleHideTemporaryActors", "Hide Temporary Actors"),
+					LOCTEXT("ToggleHideTemporaryActorsToolTip", "When enabled, hides temporary/run-time Actors."),
+					FSlateIcon(),
+					FUIAction(
+					FExecuteAction::CreateSP( this, &SSceneOutliner::ToggleHideTemporaryActors ),
+					FCanExecuteAction(),
+					FIsActionChecked::CreateSP( this, &SSceneOutliner::IsHidingTemporaryActors )
+					),
+					NAME_None,
+					EUserInterfaceActionType::ToggleButton
+				);
 
+				MenuBuilder.AddMenuEntry(
+					LOCTEXT("ToggleShowOnlyCurrentLevel", "Only in Current Level"),
+					LOCTEXT("ToggleShowOnlyCurrentLevelToolTip", "When enabled, only shows Actors that are in the Current Level."),
+					FSlateIcon(),
+					FUIAction(
+					FExecuteAction::CreateSP( this, &SSceneOutliner::ToggleShowOnlyCurrentLevel ),
+					FCanExecuteAction(),
+					FIsActionChecked::CreateSP( this, &SSceneOutliner::IsShowingOnlyCurrentLevel )
+					),
+					NAME_None,
+					EUserInterfaceActionType::ToggleButton
+				);
+			}
+			MenuBuilder.EndSection();
+
+			MenuBuilder.BeginSection("AssetThumbnails", LOCTEXT("ShowWorldHeading", "World"));
+			{
+				MenuBuilder.AddSubMenu(
+					LOCTEXT("ChooseWorldSubMenu", "Choose World"),
+					LOCTEXT("ChooseWorldSubMenuToolTip", "Choose the world to display in the outliner."),
+					FNewMenuDelegate::CreateSP(this, &SSceneOutliner::BuildWorldPickerContent)
+				);
+			}
+			MenuBuilder.EndSection();
+		}
+
+		return MenuBuilder.MakeWidget();
+	}
+
+	void SSceneOutliner::BuildWorldPickerContent(FMenuBuilder& MenuBuilder)
+	{
+		MenuBuilder.BeginSection("Worlds", LOCTEXT("WorldsHeading", "Worlds"));
+		{
 			MenuBuilder.AddMenuEntry(
-				LOCTEXT("ToggleShowOnlyCurrentLevel", "Only in Current Level"),
-				LOCTEXT("ToggleShowOnlyCurrentLevelToolTip", "When enabled, only shows Actors that are in the Current Level."),
+				LOCTEXT("AutoWorld", "Auto"),
+				LOCTEXT("AutoWorldToolTip", "Automatically pick the world to display based on context."),
 				FSlateIcon(),
 				FUIAction(
-				FExecuteAction::CreateSP( this, &SSceneOutliner::ToggleShowOnlyCurrentLevel ),
+				FExecuteAction::CreateSP( this, &SSceneOutliner::OnSelectWorld, TWeakObjectPtr<UWorld>() ),
 				FCanExecuteAction(),
-				FIsActionChecked::CreateSP( this, &SSceneOutliner::IsShowingOnlyCurrentLevel )
+				FIsActionChecked::CreateSP( this, &SSceneOutliner::IsWorldChecked, TWeakObjectPtr<UWorld>() )
 				),
 				NAME_None,
-				EUserInterfaceActionType::ToggleButton
+				EUserInterfaceActionType::RadioButton
 			);
+
+			for (const FWorldContext& Context : GEngine->GetWorldContexts())
+			{
+				UWorld* World = Context.World();
+				if (World && (World->WorldType == EWorldType::PIE || Context.WorldType == EWorldType::Editor))
+				{
+					MenuBuilder.AddMenuEntry(
+						GetWorldDescription(World),
+						LOCTEXT("ChooseWorldToolTip", "Display actors for this world."),
+						FSlateIcon(),
+						FUIAction(
+						FExecuteAction::CreateSP( this, &SSceneOutliner::OnSelectWorld, TWeakObjectPtr<UWorld>(World) ),
+						FCanExecuteAction(),
+						FIsActionChecked::CreateSP( this, &SSceneOutliner::IsWorldChecked, TWeakObjectPtr<UWorld>(World) )
+						),
+						NAME_None,
+						EUserInterfaceActionType::RadioButton
+					);
+				}
+			}
 		}
 		MenuBuilder.EndSection();
-
-		return MenuBuilder.MakeWidget();
 	}
 
-
 	/** FILTERS */
 
 	// Show Only Selected
@@ -687,16 +833,60 @@ namespace SceneOutliner
 		// while we doing this
 		TGuardValue<bool> ReentrantGuard(bIsReentrant, true);
 
-		for (const FWorldContext& Context : GEngine->GetWorldContexts())
+		SharedData->RepresentingWorld = nullptr;
+
+		// check if the user-chosen world is valid and in the editor contexts
+		if(SharedData->UserChosenWorld.IsValid())
 		{
-			if (Context.WorldType == EWorldType::PIE)
+			for (const FWorldContext& Context : GEngine->GetWorldContexts())
 			{
-				SharedData->RepresentingWorld = Context.World();
-				break;
+				if(SharedData->UserChosenWorld.Get() == Context.World())
+				{
+					SharedData->RepresentingWorld = Context.World();
+					break;
+				}
 			}
-			else if (Context.WorldType == EWorldType::Editor)
+		}
+		
+		if(SharedData->RepresentingWorld == nullptr)
+		{
+			// try to pick the most suitable world context
+
+			// ideally we want a PIE world that is standalone or the first client
+			for (const FWorldContext& Context : GEngine->GetWorldContexts())
 			{
-				SharedData->RepresentingWorld = Context.World();
+				UWorld* World = Context.World();
+				if (Context.WorldType == EWorldType::PIE)
+				{
+					if(World->GetNetMode() == NM_Standalone)
+					{
+						SharedData->RepresentingWorld = World;
+						break;
+					}
+					else if(World->GetNetMode() == NM_Client && Context.PIEInstance == 2)	// Slightly dangerous: assumes server is always PIEInstance = 1;
+					{
+						SharedData->RepresentingWorld = World;
+						break;						
+					}
+				}
+			}
+		}
+
+		if(SharedData->RepresentingWorld == nullptr)
+		{
+			// still not world so fallback to old logic where we just prefer PIE over Editor
+
+			for (const FWorldContext& Context : GEngine->GetWorldContexts())
+			{
+				if (Context.WorldType == EWorldType::PIE)
+				{
+					SharedData->RepresentingWorld = Context.World();
+					break;
+				}
+				else if (Context.WorldType == EWorldType::Editor)
+				{
+					SharedData->RepresentingWorld = Context.World();
+				}
 			}
 		}
 		
@@ -2363,6 +2553,16 @@ namespace SceneOutliner
 		}
 	}
 
+	void SSceneOutliner::OnSelectWorld(TWeakObjectPtr<UWorld> InWorld)
+	{
+		SharedData->UserChosenWorld = InWorld;
+		FullRefresh();
+	}
+
+	bool SSceneOutliner::IsWorldChecked(TWeakObjectPtr<UWorld> InWorld)
+	{
+		return (InWorld == SharedData->UserChosenWorld);
+	}
 }	// namespace SceneOutliner
 
 #undef LOCTEXT_NAMESPACE
diff --git a/Engine/Source/Editor/SceneOutliner/Private/SSceneOutliner.h b/Engine/Source/Editor/SceneOutliner/Private/SSceneOutliner.h
--- a/Engine/Source/Editor/SceneOutliner/Private/SSceneOutliner.h
+++ b/Engine/Source/Editor/SceneOutliner/Private/SSceneOutliner.h
@@ -40,6 +40,9 @@ namespace SceneOutliner
 		};
 	}
 
+	/** Get a description of a world to display in the scene outliner */
+	FText GetWorldDescription(UWorld* World);
+
 	/**
 	 * Scene Outliner widget
 	 */
@@ -320,11 +323,17 @@ namespace SceneOutliner
 		ESelectionMode::Type GetSelectionMode() const;
 
 		/** @return the content for the view button */
-		TSharedRef<SWidget> GetViewButtonContent();
+		TSharedRef<SWidget> GetViewButtonContent(bool bWorldPickerOnly);
+
+		/** Build the content for the world picker submenu */
+		void BuildWorldPickerContent(FMenuBuilder& MenuBuilder);
 
 		/** @return the foreground color for the view button */
 		FSlateColor GetViewButtonForegroundColor() const;
 
+		/** @return the foreground color for the world picker button */
+		FSlateColor GetWorldPickerForegroundColor() const;
+
 		/** The brush to use when in Editor mode */
 		const FSlateBrush* NoBorder;
 		/** The brush to use when in PIE mode */
@@ -500,6 +509,11 @@ namespace SceneOutliner
 		/** Sort the specified array of items based on the current sort column */
 		void SortItems(TArray<FTreeItemPtr>& Items) const;
 
+		/** Select the world we want to view */
+		void OnSelectWorld(TWeakObjectPtr<UWorld> InWorld);
+
+		/** Display a checkbox next to the world we are viewing */
+		bool IsWorldChecked(TWeakObjectPtr<UWorld> InWorld);
 	};
 
 }		// namespace SceneOutliner
diff --git a/Engine/Source/Editor/SceneOutliner/Private/SceneOutlinerItemLabelColumn.cpp b/Engine/Source/Editor/SceneOutliner/Private/SceneOutlinerItemLabelColumn.cpp
--- a/Engine/Source/Editor/SceneOutliner/Private/SceneOutlinerItemLabelColumn.cpp
+++ b/Engine/Source/Editor/SceneOutliner/Private/SceneOutlinerItemLabelColumn.cpp
@@ -400,9 +400,15 @@ private:
 	FText GetTooltipText() const
 	{
 		auto Item = TreeItemPtr.Pin();
-		FText PersistentLevelDisplayName = Item.IsValid() ? FText::FromString(Item->GetDisplayString()) : FText();
-		
-		return FText::Format(LOCTEXT("WorldLabel_Tooltip", "The world settings for {0}, double-click to edit"), PersistentLevelDisplayName);
+		FText PersistentLevelDisplayName = Item.IsValid() ? FText::FromString(Item->GetWorldName()) : FText();
+		if(Item->CanInteract())
+		{
+			return FText::Format(LOCTEXT("WorldLabel_Tooltip", "The world settings for {0}, double-click to edit"), PersistentLevelDisplayName);
+		}
+		else
+		{
+			return FText::Format(LOCTEXT("WorldLabel_TooltipNonInteractive", "The world {0}"), PersistentLevelDisplayName);
+		}
 	}
 
 	FSlateColor GetForegroundColor() const
diff --git a/Engine/Source/Editor/SceneOutliner/Private/WorldTreeItem.cpp b/Engine/Source/Editor/SceneOutliner/Private/WorldTreeItem.cpp
--- a/Engine/Source/Editor/SceneOutliner/Private/WorldTreeItem.cpp
+++ b/Engine/Source/Editor/SceneOutliner/Private/WorldTreeItem.cpp
@@ -5,6 +5,7 @@
 #include "WorldTreeItem.h"
 #include "LevelEditor.h"
 #include "EditorActorFolders.h"
+#include "SceneOutliner.h"
 
 #define LOCTEXT_NAMESPACE "SceneOutliner_WorldTreeItem"
 
@@ -44,6 +45,15 @@ FTreeItemID FWorldTreeItem::GetID() const
 }
 
 FString FWorldTreeItem::GetDisplayString() const
+{
+	if (UWorld* WorldPtr = World.Get())
+	{
+		return SceneOutliner::GetWorldDescription(WorldPtr).ToString();
+	}
+	return FString();
+}
+
+FString FWorldTreeItem::GetWorldName() const
 {
 	if (UWorld* WorldPtr = World.Get())
 	{
@@ -59,6 +69,10 @@ int32 FWorldTreeItem::GetTypeSortPriority() const
 
 bool FWorldTreeItem::CanInteract() const
 {
+	if (UWorld* WorldPtr = World.Get())
+	{
+		return Flags.bInteractive && WorldPtr->WorldType == EWorldType::Editor;
+	}
 	return Flags.bInteractive;
 }
 
diff --git a/Engine/Source/Editor/SceneOutliner/Public/SceneOutlinerPublicTypes.h b/Engine/Source/Editor/SceneOutliner/Public/SceneOutlinerPublicTypes.h
--- a/Engine/Source/Editor/SceneOutliner/Public/SceneOutlinerPublicTypes.h
+++ b/Engine/Source/Editor/SceneOutliner/Public/SceneOutlinerPublicTypes.h
@@ -162,9 +162,13 @@ namespace SceneOutliner
 		/** The world that we are representing */
 		UWorld* RepresentingWorld;
 
+		/** The world the user has chosen to display */
+		TWeakObjectPtr<UWorld> UserChosenWorld;
+
 		FSharedOutlinerData()
 			: bRepresentingPlayWorld(false)
 			, RepresentingWorld(nullptr)
+			, UserChosenWorld(nullptr)
 		{}
 	};
 
diff --git a/Engine/Source/Editor/SceneOutliner/Public/WorldTreeItem.h b/Engine/Source/Editor/SceneOutliner/Public/WorldTreeItem.h
--- a/Engine/Source/Editor/SceneOutliner/Public/WorldTreeItem.h
+++ b/Engine/Source/Editor/SceneOutliner/Public/WorldTreeItem.h
@@ -56,6 +56,9 @@ namespace SceneOutliner
 		/** Called to drop the specified objects on this item. Only called if ValidateDrop() allows. */
 		virtual void OnDrop(FDragDropPayload& DraggedObjects, UWorld& World, const FDragValidationInfo& ValidationInfo, TSharedRef<SWidget> DroppedOnWidget) override;
 
+		/** Get just the name of the world, for tooltip use */
+		FString GetWorldName() const;
+
 	private:
 
 		/** Create a new folder at the root of this world */
diff --git a/Engine/Source/Editor/SequenceRecorder/Private/ActorRecording.cpp b/Engine/Source/Editor/SequenceRecorder/Private/ActorRecording.cpp
new file mode 100644
--- /dev/null
+++ b/Engine/Source/Editor/SequenceRecorder/Private/ActorRecording.cpp
@@ -0,0 +1,607 @@
+// Copyright 1998-2016 Epic Games, Inc. All Rights Reserved.
+
+#include "SequenceRecorderPrivatePCH.h"
+#include "ActorRecording.h"
+#include "SequenceRecorderSettings.h"
+#include "SequenceRecorderUtils.h"
+#include "IMovieScenePropertyRecorder.h"
+#include "MovieSceneSpawnPropertyRecorder.h"
+#include "MovieScene3DTransformPropertyRecorder.h"
+#include "MovieSceneVisibilityPropertyRecorder.h"
+#include "MovieSceneAnimationPropertyRecorder.h"
+#include "MovieSceneParticleTrackPropertyRecorder.h"
+#include "MovieScene3DAttachPropertyRecorder.h"
+#include "AssetSelection.h"
+#include "MovieSceneSkeletalAnimationTrack.h"
+#include "MovieSceneSkeletalAnimationSection.h"
+#include "MovieSceneSpawnTrack.h"
+#include "MovieScene3DTransformTrack.h"
+#include "MovieScene3DTransformSection.h"
+#include "MovieSceneBoolSection.h"
+#include "KismetEditorUtilities.h"
+#include "BlueprintEditorUtils.h"
+#include "MovieSceneFolder.h"
+#include "CameraRig_Crane.h"
+#include "CameraRig_Rail.h"
+#include "Animation/SkeletalMeshActor.h"
+
+static const FName SequencerActorTag(TEXT("SequencerActor"));
+
+UActorRecording::UActorRecording(const FObjectInitializer& ObjectInitializer)
+	: Super(ObjectInitializer)
+{
+	bWasSpawnedPostRecord = false;
+	MeshComponentUpdateFlag = EMeshComponentUpdateFlag::AlwaysTickPoseAndRefreshBones;
+	bEnableUpdateRateOptimizations = false;
+	Guid.Invalidate();
+	bNewComponentAddedWhileRecording = false;
+
+	if(!HasAnyFlags(RF_ClassDefaultObject))
+	{
+		const USequenceRecorderSettings* Settings = GetDefault<USequenceRecorderSettings>();
+		AnimationSettings = Settings->DefaultAnimationSettings;
+	}
+}
+
+bool UActorRecording::IsRelevantForRecording(AActor* Actor)
+{
+	// don't record actors that sequencer has spawned itself!
+	if(Actor->ActorHasTag(SequencerActorTag))
+	{
+		return false;
+	}
+
+	TInlineComponentArray<USceneComponent*> SceneComponents(Actor);
+
+	for(USceneComponent* SceneComponent : SceneComponents)
+	{
+		if(SceneComponent->IsA<USkeletalMeshComponent>() ||
+			SceneComponent->IsA<UStaticMeshComponent>() ||
+			SceneComponent->IsA<UParticleSystemComponent>() || 
+			SceneComponent->IsA<ULightComponent>())
+		{
+			return true;
+		}
+	}
+
+	return false;
+}
+
+bool UActorRecording::StartRecording(ULevelSequence* CurrentSequence, float CurrentSequenceTime)
+{
+	bNewComponentAddedWhileRecording = false;
+
+	if(ActorToRecord.IsValid())
+	{
+		if(CurrentSequence != nullptr)
+		{
+			StartRecordingActorProperties(CurrentSequence, CurrentSequenceTime);
+		}
+		else
+		{
+			TSharedPtr<FMovieSceneAnimationPropertyRecorder> AnimationRecorder = MakeShareable(new FMovieSceneAnimationPropertyRecorder(AnimationSettings, TargetAnimation.Get()));
+			AnimationRecorder->CreateSection(ActorToRecord.Get(), nullptr, FGuid(), 0.0f, true);
+			PropertyRecorders.Add(AnimationRecorder);			
+		}
+	}
+
+	return true;
+}
+
+static FString GetUniqueSpawnableName(UMovieScene* MovieScene, const FString& BaseName)
+{
+	FString BlueprintName = BaseName;
+	auto DuplName = [&](FMovieSceneSpawnable& InSpawnable)
+	{
+		return InSpawnable.GetName() == BlueprintName;
+	};
+
+	int32 Index = 2;
+	FString UniqueString;
+	while (MovieScene->FindSpawnable(DuplName))
+	{
+		BlueprintName.RemoveFromEnd(UniqueString);
+		UniqueString = FString::Printf(TEXT(" (%d)"), Index++);
+		BlueprintName += UniqueString;
+	}
+
+	return BlueprintName;
+}
+
+void UActorRecording::GetSceneComponents(TArray<USceneComponent*>& OutArray, bool bIncludeNonCDO/*=true*/)
+{
+	// it is not enough to just go through the owned components array here
+	// we need to traverse the scene component hierarchy as well, as some components may be 
+	// owned by other actors (e.g. for pooling) and some may not be part of the hierarchy
+	if(ActorToRecord.IsValid())
+	{
+		USceneComponent* RootComponent = ActorToRecord->GetRootComponent();
+		if(RootComponent)
+		{
+			// note: GetChildrenComponents clears array!
+			RootComponent->GetChildrenComponents(true, OutArray);
+			OutArray.Add(RootComponent);
+		}
+
+		// add owned components that are *not* part of the hierarchy
+		TInlineComponentArray<USceneComponent*> OwnedComponents(ActorToRecord.Get());
+		for(USceneComponent* OwnedComponent : OwnedComponents)
+		{
+			if(OwnedComponent->AttachParent == nullptr && OwnedComponent != RootComponent)
+			{
+				OutArray.Add(OwnedComponent);
+			}
+		}
+
+		if(!bIncludeNonCDO)
+		{
+			AActor* CDO = Cast<AActor>(ActorToRecord->GetClass()->GetDefaultObject());
+
+			auto ShouldRemovePredicate = [&](UActorComponent* PossiblyRemovedComponent)
+				{
+					// try to find a component with this name in the CDO
+					UActorComponent* const* FoundComponent = CDO->GetComponents().FindByPredicate([&](UActorComponent* SearchComponent) 
+					{ 
+						return SearchComponent->GetClass() == PossiblyRemovedComponent->GetClass() &&
+							   SearchComponent->GetFName() == PossiblyRemovedComponent->GetFName(); 
+					} );
+
+					// remove if its not found
+					return FoundComponent == nullptr;
+				};
+
+			OutArray.RemoveAllSwap(ShouldRemovePredicate);
+		}
+	}
+}
+
+void UActorRecording::SyncTrackedComponents(bool bIncludeNonCDO/*=true*/)
+{
+	TArray<USceneComponent*> NewComponentArray;
+	GetSceneComponents(NewComponentArray, bIncludeNonCDO);
+
+	TrackedComponents.Reset(NewComponentArray.Num());
+	for(USceneComponent* SceneComponent : NewComponentArray)
+	{
+		TrackedComponents.Add(SceneComponent);
+	}
+}
+
+void UActorRecording::InvalidateObjectToRecord()
+{
+	ActorToRecord = nullptr;
+	for(auto& PropertyRecorder : PropertyRecorders)
+	{
+		PropertyRecorder->InvalidateObjectToRecord();
+	}
+}
+
+bool UActorRecording::ValidComponent(USceneComponent* SceneComponent) const
+{
+	if(SceneComponent != nullptr)
+	{
+		const bool bIsValidComponent = 
+			SceneComponent->IsA<USkeletalMeshComponent>() ||
+			SceneComponent->IsA<UStaticMeshComponent>() ||
+			SceneComponent->IsA<UParticleSystemComponent>();
+
+		return bIsValidComponent;
+	}
+
+	return false;
+}
+
+void UActorRecording::FindOrAddFolder(UMovieScene* MovieScene)
+{
+	check(ActorToRecord.IsValid());
+
+	FName FolderName(NAME_None);
+	if(ActorToRecord.Get()->IsA<ACharacter>() || ActorToRecord.Get()->IsA<ASkeletalMeshActor>())
+	{
+		FolderName = TEXT("Characters");
+	}
+	else if(ActorToRecord.Get()->IsA<ACameraActor>() || ActorToRecord.Get()->IsA<ACameraRig_Crane>() || ActorToRecord.Get()->IsA<ACameraRig_Rail>())
+	{
+		FolderName = TEXT("Cameras");
+	}
+	else
+	{
+		FolderName = TEXT("Misc");
+	}
+
+	// look for a folder to put us in
+	UMovieSceneFolder* FolderToUse = nullptr;
+	for(UMovieSceneFolder* Folder : MovieScene->GetRootFolders())
+	{
+		if(Folder->GetFolderName() == FolderName)
+		{
+			FolderToUse = Folder;
+			break;
+		}
+	}
+
+	if(FolderToUse == nullptr)
+	{
+		FolderToUse = NewObject<UMovieSceneFolder>(MovieScene, NAME_None, RF_Transactional);
+		FolderToUse->SetFolderName(FolderName);
+		MovieScene->GetRootFolders().Add(FolderToUse);
+	}
+
+	FolderToUse->AddChildObjectBinding(Guid);
+}
+
+void UActorRecording::StartRecordingActorProperties(ULevelSequence* CurrentSequence, float CurrentSequenceTime)
+{
+	if(CurrentSequence != nullptr)
+	{
+		// set up our spawnable for this actor
+		UMovieScene* MovieScene = CurrentSequence->GetMovieScene();
+
+		AActor* Actor = ActorToRecord.Get();
+		FString BlueprintName = GetUniqueSpawnableName(MovieScene, Actor->GetName());
+
+		UClass* ActorClass = Actor->GetClass();
+			
+		UBlueprint* Blueprint = nullptr;
+
+		if(ActorClass->IsChildOf<AWorldSettings>())
+		{
+			// we need to create our own substitute blueprint for world settings actors as they are not blueprintable by default
+			Blueprint = FKismetEditorUtilities::CreateBlueprint(AActor::StaticClass(), MovieScene, *BlueprintName, EBlueprintType::BPTYPE_Normal, UBlueprint::StaticClass(), UBlueprintGeneratedClass::StaticClass(), NAME_Name);
+		}
+		else
+		{
+			UBlueprint* NewBlueprint = FKismetEditorUtilities::CreateBlueprint(ActorClass, MovieScene, *BlueprintName, EBlueprintType::BPTYPE_Normal, UBlueprint::StaticClass(), UBlueprintGeneratedClass::StaticClass(), NAME_Name);
+			if(NewBlueprint)
+			{
+				// if the duplicated BP has any skel mesh components, set them to single anim instance
+				AActor* CDO = CastChecked<AActor>(NewBlueprint->GeneratedClass->ClassDefaultObject);
+				if(CDO)
+				{
+					TInlineComponentArray<USkeletalMeshComponent*> SkeletalMeshComponents;
+					CDO->GetComponents(SkeletalMeshComponents);	
+					for(USkeletalMeshComponent* SkeletalMeshComponent : SkeletalMeshComponents)
+					{
+						SkeletalMeshComponent->SetAnimationMode(EAnimationMode::AnimationSingleNode);
+						SkeletalMeshComponent->bEnableUpdateRateOptimizations = false;
+						SkeletalMeshComponent->MeshComponentUpdateFlag = EMeshComponentUpdateFlag::AlwaysTickPoseAndRefreshBones;
+						SkeletalMeshComponent->ForcedLodModel = 1;
+					}
+				}
+
+				Blueprint = NewBlueprint;
+			}
+		}
+
+		if(Blueprint)
+		{
+			Guid = MovieScene->AddSpawnable(BlueprintName, Blueprint);
+		}
+
+		// now add tracks to record
+		if(Guid.IsValid())
+		{
+			// add our folder
+			FindOrAddFolder(MovieScene);
+
+			// force set recording to record translations as we need this with no animation
+			ActorSettings.bRecordTransforms = true;
+
+			// grab components so we can track attachments
+			// don't include non-CDO here as they wont be part of our initial BP (duplicated above)
+			// we will catch these 'extra' components on the first tick
+			const bool bIncludeNonCDO = false;
+			SyncTrackedComponents(bIncludeNonCDO);
+
+			TInlineComponentArray<USceneComponent*> SceneComponents(ActorToRecord.Get());
+
+			// check if components need recording
+			TInlineComponentArray<USceneComponent*> ValidSceneComponents;
+			for(TWeakObjectPtr<USceneComponent>& SceneComponent : TrackedComponents)
+			{
+				if(ValidComponent(SceneComponent.Get()))
+				{
+					ValidSceneComponents.Add(SceneComponent.Get());
+
+					// add all parent components too
+					TArray<USceneComponent*> ParentComponents;
+					SceneComponent->GetParentComponents(ParentComponents);
+					for(USceneComponent* ParentComponent : ParentComponents)
+					{	
+						ValidSceneComponents.AddUnique(ParentComponent);
+					}
+				}
+			}
+
+			TSharedPtr<FMovieSceneAnimationPropertyRecorder> FirstAnimRecorder = nullptr;
+			for(USceneComponent* SceneComponent : ValidSceneComponents)
+			{
+				TSharedPtr<FMovieSceneAnimationPropertyRecorder> AnimRecorder = StartRecordingComponentProperties(SceneComponent->GetFName(), SceneComponent, ActorToRecord.Get(), CurrentSequence, CurrentSequenceTime);
+				if(!FirstAnimRecorder.IsValid() && AnimRecorder.IsValid())
+				{
+					FirstAnimRecorder = AnimRecorder;
+				}
+			}
+
+			// add a spawn track so we actually get spawned
+			TSharedPtr<FMovieSceneSpawnPropertyRecorder> SpawnRecorder = MakeShareable(new FMovieSceneSpawnPropertyRecorder);
+			SpawnRecorder->CreateSection(ActorToRecord.Get(), MovieScene, Guid, CurrentSequenceTime, true);
+			PropertyRecorders.Add(SpawnRecorder);
+
+			// we need to create a transform track even if we arent recording transforms
+			TSharedPtr<FMovieScene3DTransformPropertyRecorder> TransformRecorder = MakeShareable(new FMovieScene3DTransformPropertyRecorder(FirstAnimRecorder));
+			TransformRecorder->CreateSection(ActorToRecord.Get(), MovieScene, Guid, CurrentSequenceTime, ActorSettings.bRecordTransforms);
+			PropertyRecorders.Add(TransformRecorder);
+
+			// add a visibility track 
+			TSharedPtr<FMovieSceneVisibilityPropertyRecorder> VisibilityRecorder = MakeShareable(new FMovieSceneVisibilityPropertyRecorder);
+			VisibilityRecorder->CreateSection(ActorToRecord.Get(), MovieScene, Guid, CurrentSequenceTime, ActorSettings.bRecordVisibility);
+			PropertyRecorders.Add(VisibilityRecorder);
+
+			// add an attachment track
+			TSharedPtr<FMovieScene3DAttachPropertyRecorder> AttachRecorder = MakeShareable(new FMovieScene3DAttachPropertyRecorder);
+			AttachRecorder->CreateSection(ActorToRecord.Get(), MovieScene, Guid, CurrentSequenceTime, true);
+			PropertyRecorders.Add(AttachRecorder);			
+		}
+	}
+}
+
+static FString GetUniquePosessableName(UMovieScene* MovieScene, const FString& BaseName)
+{
+	FString BlueprintName = BaseName;
+	auto DuplName = [&](FMovieScenePossessable& InPosessable)
+	{
+		return InPosessable.GetName() == BlueprintName;
+	};
+
+	int32 Index = 2;
+	FString UniqueString;
+	while (MovieScene->FindPossessable(DuplName))
+	{
+		BlueprintName.RemoveFromEnd(UniqueString);
+		UniqueString = FString::Printf(TEXT(" (%d)"), Index++);
+		BlueprintName += UniqueString;
+	}
+
+	return BlueprintName;
+}
+
+TSharedPtr<FMovieSceneAnimationPropertyRecorder> UActorRecording::StartRecordingComponentProperties(const FName& BindingName, USceneComponent* SceneComponent, UObject* BindingContext, ULevelSequence* CurrentSequence, float CurrentSequenceTime)
+{
+	// first create a possessable for this component to be controlled by
+	UMovieScene* OwnerMovieScene = CurrentSequence->GetMovieScene();
+
+	FString Name = GetUniquePosessableName(OwnerMovieScene, BindingName.ToString());
+
+	const FGuid PossessableGuid = OwnerMovieScene->AddPossessable(Name, SceneComponent->GetClass());
+
+	// Set up parent/child guids for possessables within spawnables
+	FMovieScenePossessable* ChildPossessable = OwnerMovieScene->FindPossessable(PossessableGuid);
+	if (ensure(ChildPossessable))
+	{
+		ChildPossessable->SetParent(Guid);
+	}
+
+	FMovieSceneSpawnable* ParentSpawnable = OwnerMovieScene->FindSpawnable(Guid);
+	if (ParentSpawnable)
+	{
+		ParentSpawnable->AddChildPossessable(PossessableGuid);
+	}
+
+	FLevelSequenceObjectReference ObjectReference(FUniqueObjectGuid(), BindingName.ToString());
+
+	CurrentSequence->BindPossessableObject(PossessableGuid, ObjectReference);
+
+	TSharedPtr<FMovieSceneAnimationPropertyRecorder> AnimationRecorder;
+	if(SceneComponent->IsA<USkeletalMeshComponent>())
+	{
+		AnimationRecorder = MakeShareable(new FMovieSceneAnimationPropertyRecorder(AnimationSettings, TargetAnimation.Get()));
+		AnimationRecorder->CreateSection(SceneComponent, OwnerMovieScene, PossessableGuid, CurrentSequenceTime, true);
+		PropertyRecorders.Add(AnimationRecorder);
+	}
+	else if(SceneComponent->IsA<UParticleSystemComponent>())
+	{
+		TSharedPtr<FMovieSceneParticleTrackPropertyRecorder> ParticleTrackRecorder = MakeShareable(new FMovieSceneParticleTrackPropertyRecorder());
+		ParticleTrackRecorder->CreateSection(SceneComponent, OwnerMovieScene, PossessableGuid, CurrentSequenceTime, true);
+		PropertyRecorders.Add(ParticleTrackRecorder);
+	}
+
+	// Dont record the root component transforms as this will be taken into account by the actor transform track
+	// Also dont record transforms of skeletal mesh components as they will be taken into account in the actor transform
+	bool bIsCharacterSkelMesh = false;
+	if(SceneComponent->IsA<USkeletalMeshComponent>() && SceneComponent->GetOwner()->IsA<ACharacter>())
+	{
+		ACharacter* Character = CastChecked<ACharacter>(SceneComponent->GetOwner());
+		bIsCharacterSkelMesh = SceneComponent == Character->GetMesh();
+	}
+
+	if(SceneComponent != SceneComponent->GetOwner()->GetRootComponent() && !bIsCharacterSkelMesh)
+	{
+		TSharedPtr<FMovieScene3DTransformPropertyRecorder> TransformRecorder = MakeShareable(new FMovieScene3DTransformPropertyRecorder);
+		TransformRecorder->CreateSection(SceneComponent, OwnerMovieScene, PossessableGuid, CurrentSequenceTime, true);
+		PropertyRecorders.Add(TransformRecorder);
+	}
+
+	TSharedPtr<FMovieSceneVisibilityPropertyRecorder> VisibilityRecorder = MakeShareable(new FMovieSceneVisibilityPropertyRecorder);
+	VisibilityRecorder->CreateSection(SceneComponent, OwnerMovieScene, PossessableGuid, CurrentSequenceTime, true);
+	PropertyRecorders.Add(VisibilityRecorder);	
+
+	return AnimationRecorder;
+}
+
+void UActorRecording::Tick(float DeltaSeconds, ULevelSequence* CurrentSequence, float CurrentSequenceTime)
+{
+	if (IsRecording())
+	{
+		if(CurrentSequence)
+		{
+			// check our components to see if they have changed
+			static TArray<USceneComponent*> SceneComponents;
+			GetSceneComponents(SceneComponents);
+
+			if(TrackedComponents.Num() != SceneComponents.Num())
+			{
+				StartRecordingNewComponents(CurrentSequence, CurrentSequenceTime);
+			}
+		}
+
+		for(auto& PropertyRecorder : PropertyRecorders)
+		{
+			PropertyRecorder->Record(CurrentSequenceTime);
+		}
+	}
+}
+
+bool UActorRecording::StopRecording(ULevelSequence* CurrentSequence)
+{
+	FString ActorName;
+	if(CurrentSequence)
+	{
+		UMovieScene* MovieScene = CurrentSequence->GetMovieScene();
+		check(MovieScene);
+
+		FMovieSceneSpawnable* Spawnable = MovieScene->FindSpawnable(Guid);
+		if(Spawnable)
+		{
+			ActorName = Spawnable->GetName();
+		}
+	}
+
+	FScopedSlowTask SlowTask((float)PropertyRecorders.Num() + 1.0f, FText::Format(NSLOCTEXT("SequenceRecorder", "ProcessingActor", "Processing Actor {0}"), FText::FromString(ActorName)));
+
+	// stop property recorders
+	for(auto& PropertyRecorder : PropertyRecorders)
+	{
+		SlowTask.EnterProgressFrame();
+
+		PropertyRecorder->FinalizeSection();
+	}
+
+	SlowTask.EnterProgressFrame();
+
+	PropertyRecorders.Empty();
+
+	if(bNewComponentAddedWhileRecording && CurrentSequence)
+	{
+		UMovieScene* MovieScene = CurrentSequence->GetMovieScene();
+		check(MovieScene);
+
+		FMovieSceneSpawnable* Spawnable = MovieScene->FindSpawnable(Guid);
+		if(Spawnable)
+		{
+			UBlueprint* Blueprint = CastChecked<UBlueprint>(Spawnable->GetClass()->ClassGeneratedBy);
+			FBlueprintEditorUtils::MarkBlueprintAsStructurallyModified(Blueprint);
+		}
+	}
+
+	return true;
+}
+
+bool UActorRecording::IsRecording() const
+{
+	return ActorToRecord.IsValid() && PropertyRecorders.Num() > 0;
+}
+
+static FName FindParentComponentOwnerClassName(USceneComponent* SceneComponent, UBlueprint* Blueprint)
+{
+	if(SceneComponent->AttachParent)
+	{
+		FName AttachName = SceneComponent->AttachParent->GetFName();
+
+		// see if we can find this component in the BP inheritance hierarchy
+		while(Blueprint)
+		{
+			if(Blueprint->SimpleConstructionScript->FindSCSNode(AttachName) != nullptr)
+			{
+				return Blueprint->GetFName();
+			}
+
+			Blueprint = Cast<UBlueprint>(Blueprint->GeneratedClass->GetSuperClass()->ClassGeneratedBy);
+		}
+	}
+
+	return NAME_None;
+}
+
+void UActorRecording::StartRecordingNewComponents(ULevelSequence* CurrentSequence, float CurrentSequenceTime)
+{
+	if (ActorToRecord.IsValid())
+	{
+		// find the new component(s)
+		TArray<USceneComponent*> NewComponents;
+		TArray<USceneComponent*> SceneComponents;
+		GetSceneComponents(SceneComponents);
+		for(USceneComponent* SceneComponent : SceneComponents)
+		{
+			if(ValidComponent(SceneComponent))
+			{
+				TWeakObjectPtr<USceneComponent> WeakSceneComponent(SceneComponent);
+				int32 FoundIndex = TrackedComponents.Find(WeakSceneComponent);
+				if(FoundIndex == INDEX_NONE)
+				{
+					// new component!
+					NewComponents.Add(SceneComponent);
+				}
+			}
+		}
+
+		UMovieScene* MovieScene = CurrentSequence->GetMovieScene();
+		check(MovieScene);
+
+		FMovieSceneSpawnable* Spawnable = MovieScene->FindSpawnable(Guid);
+		check(Spawnable);
+
+		UBlueprint* Blueprint = CastChecked<UBlueprint>(Spawnable->GetClass()->ClassGeneratedBy);
+
+		for(USceneComponent* SceneComponent : NewComponents)
+		{
+			// new component, so we need to add this to our BP if it didn't come from SCS
+			FName NewName;
+			if(SceneComponent->CreationMethod != EComponentCreationMethod::SimpleConstructionScript)
+			{
+				NewName = *FString::Printf(TEXT("Dynamic%s"), *SceneComponent->GetFName().ToString());
+				USCS_Node* NewSCSNode = Blueprint->SimpleConstructionScript->CreateNode(SceneComponent->GetClass(), NewName); 
+				UEditorEngine::CopyPropertiesForUnrelatedObjects(SceneComponent, NewSCSNode->ComponentTemplate);
+				NewSCSNode->ComponentTemplate->CreationMethod = EComponentCreationMethod::SimpleConstructionScript;
+
+				// look for a similar attach parent in the current structure
+				bool bAdded = false;
+				if(SceneComponent->AttachParent != nullptr)
+				{
+					FName AttachName = SceneComponent->AttachParent->GetFName();
+					NewSCSNode->ParentComponentOrVariableName = AttachName;
+					NewSCSNode->bIsParentComponentNative = SceneComponent->AttachParent->CreationMethod == EComponentCreationMethod::Native;
+					NewSCSNode->ParentComponentOwnerClassName = FindParentComponentOwnerClassName(SceneComponent, Blueprint);
+
+					USCS_Node* ParentSCSNode = Blueprint->SimpleConstructionScript->FindSCSNode(AttachName);
+					if(ParentSCSNode)
+					{
+						ParentSCSNode->AddChildNode(NewSCSNode);
+						bAdded = true;
+					}
+				}
+			
+				// nothing found - add to root
+				if(!bAdded)
+				{
+					Blueprint->SimpleConstructionScript->AddNode(NewSCSNode);
+				}
+
+				if(SceneComponent->AttachSocketName != NAME_None)
+				{
+					NewSCSNode->AttachToName = SceneComponent->AttachSocketName;
+				}
+			}
+			else
+			{
+				NewName = SceneComponent->GetFName();
+			}
+
+			StartRecordingComponentProperties(NewName, SceneComponent, ActorToRecord.Get(), CurrentSequence, CurrentSequenceTime);
+
+			bNewComponentAddedWhileRecording = true;
+		}
+		
+		SyncTrackedComponents();
+	}
+}
diff --git a/Engine/Source/Editor/SequenceRecorder/Private/ActorRecording.h b/Engine/Source/Editor/SequenceRecorder/Private/ActorRecording.h
new file mode 100644
--- /dev/null
+++ b/Engine/Source/Editor/SequenceRecorder/Private/ActorRecording.h
@@ -0,0 +1,120 @@
+// Copyright 1998-2016 Epic Games, Inc. All Rights Reserved.
+
+#pragma once
+
+#include "Animation/AnimationRecordingSettings.h"
+#include "Components/SkinnedMeshComponent.h"
+
+#include "ActorRecording.generated.h"
+
+USTRUCT()
+struct FActorRecordingSettings
+{
+	GENERATED_BODY()
+
+	FActorRecordingSettings()
+		: bRecordTransforms(false)
+		, bRecordVisibility(true)
+	{}
+
+public:
+	/** 
+	 * Whether to record actor transforms. This can be useful if you want the actors to end up in specific locations after the sequence. 
+	 * By default we rely on animations to provide transforms, but this can be changed using the "Record In World Space" animation setting.
+	 */
+	UPROPERTY(EditAnywhere, Category = "Actor Recording")
+	bool bRecordTransforms;
+
+	/** Whether to record actor visibility. */
+	UPROPERTY(EditAnywhere, Category = "Actor Recording")
+	bool bRecordVisibility;
+};
+
+UCLASS(MinimalAPI, Transient)
+class UActorRecording : public UObject
+{
+	GENERATED_UCLASS_BODY()
+
+public:
+	/** Check whether it is worth recording this actor - i.e. is it going to affect the end result of the sequence */
+	static bool IsRelevantForRecording(AActor* Actor);
+
+	/** Start this queued recording. Sequence can be nullptr */
+	bool StartRecording(class ULevelSequence* CurrentSequence = nullptr, float CurrentSequenceTime = 0.0f);
+
+	/** Stop this recording. Has no effect if we are not currently recording. Sequence can be nullptr */
+	bool StopRecording(class ULevelSequence* CurrentSequence = nullptr);
+
+	/** Tick this recording */
+	void Tick(float DeltaSeconds, ULevelSequence* CurrentSequence = nullptr, float CurrentSequenceTime = 0.0f);
+
+	/** Whether we are currently recording */
+	bool IsRecording() const;
+
+	/** Simulate a de-spawned actor */
+	void InvalidateObjectToRecord();
+
+private:
+	/** Check component valididty for recording */
+	bool ValidComponent(USceneComponent* SceneComponent) const;
+
+	/** Adds us to a folder for better sequence organization */
+	void FindOrAddFolder(UMovieScene* MovieScene);
+
+	/** Start recording actor properties to a sequence */
+	void StartRecordingActorProperties(ULevelSequence* CurrentSequence, float CurrentSequenceTime);
+
+	/** Start recording component properties to a sequence */
+	TSharedPtr<class FMovieSceneAnimationPropertyRecorder> StartRecordingComponentProperties(const FName& BindingName, USceneComponent* SceneComponent, UObject* BindingContext, ULevelSequence* CurrentSequence, float CurrentSequenceTime);
+
+	/** Start recording components that are added at runtime */
+	void StartRecordingNewComponents(ULevelSequence* CurrentSequence, float CurrentSequenceTime);
+
+	/** Helper function to grab all scene components in the actor's hierarchy */
+	void GetSceneComponents(TArray<USceneComponent*>& OutArray, bool bIncludeNonCDO = true);
+
+	/** Sync up tracked components with the actor */
+	void SyncTrackedComponents(bool bIncludeNonCDO = true);
+
+public:
+	/** The actor we want to record */
+	UPROPERTY(EditAnywhere, Category = "Actor Recording")
+	TLazyObjectPtr<AActor> ActorToRecord;
+
+	UPROPERTY(EditAnywhere, Category = "Actor Recording")
+	FActorRecordingSettings ActorSettings;
+
+	/** Whether we should specify the target animation or auto-create it */
+	UPROPERTY(EditAnywhere, Category = "Animation Recording")
+	bool bSpecifyTargetAnimation;
+
+	/** The target animation we want to record to */
+	UPROPERTY(EditAnywhere, Category = "Animation Recording", meta=(EditCondition = "bSpecifyTargetAnimation"))
+	TWeakObjectPtr<UAnimSequence> TargetAnimation;
+
+	/** The settings to apply to this actor's animation */
+	UPROPERTY(EditAnywhere, Category = "Animation Recording")
+	FAnimationRecordingSettings AnimationSettings;
+
+	/** Guid that identifies our spawnable in a recorded sequence */
+	FGuid Guid;
+
+	/** Whether this actor recording was triggered from an acotr spawn */
+	bool bWasSpawnedPostRecord;
+
+private:
+	/** Used to store/restore update flag when recording */
+	EMeshComponentUpdateFlag::Type MeshComponentUpdateFlag;
+
+	/** Used to store/restore URO when recording */
+	bool bEnableUpdateRateOptimizations;
+
+	/** This actor's current set of property recorders */
+	TArray<TSharedPtr<class IMovieScenePropertyRecorder>> PropertyRecorders;
+
+	/** Track components to check if any have changed */
+	TArray<TWeakObjectPtr<USceneComponent>> TrackedComponents;
+
+	/** Flag to track whether we created new components */
+	bool bNewComponentAddedWhileRecording;
+};
\ No newline at end of file
