Copying //UE4/Dev-Sequencer to //UE4/Main

diff --git a/Engine/Source/Editor/SequenceRecorder/Private/AnimationRecorder.cpp b/Engine/Source/Editor/SequenceRecorder/Private/AnimationRecorder.cpp
new file mode 100644
--- /dev/null
+++ b/Engine/Source/Editor/SequenceRecorder/Private/AnimationRecorder.cpp
@@ -0,0 +1,878 @@
+// Copyright 1998-2016 Epic Games, Inc. All Rights Reserved.
+
+#include "SequenceRecorderPrivatePCH.h"
+#include "AnimationRecorder.h"
+#include "Animation/AnimCompress.h"
+#include "Animation/AnimCompress_BitwiseCompressOnly.h"
+#include "SCreateAnimationDlg.h"
+#include "Runtime/AssetRegistry/Public/AssetRegistryModule.h"
+#include "SNotificationList.h"
+#include "NotificationManager.h"
+#include "EngineLogs.h"
+#include "Animation/AnimationSettings.h"
+#include "Toolkits/AssetEditorManager.h"
+#include "ActorRecording.h"
+
+#define LOCTEXT_NAMESPACE "FAnimationRecorder"
+
+/////////////////////////////////////////////////////
+
+FAnimationRecorder::FAnimationRecorder()
+	: AnimationObject(nullptr)
+	, bRecordLocalToWorld(false)
+	, bAutoSaveAsset(false)
+
+{
+	SetSampleRateAndLength(FAnimationRecordingSettings::DefaultSampleRate, FAnimationRecordingSettings::DefaultMaximumLength);
+}
+
+FAnimationRecorder::~FAnimationRecorder()
+{
+	StopRecord(false);
+}
+
+void FAnimationRecorder::SetSampleRateAndLength(float SampleRateHz, float LengthInSeconds)
+{
+	if (SampleRateHz <= 0.f)
+	{
+		// invalid rate passed in, fall back to default
+		SampleRateHz = FAnimationRecordingSettings::DefaultSampleRate;
+	}
+
+	if (LengthInSeconds <= 0.f)
+	{
+		// invalid length passed in, default to unbounded
+		SampleRateHz = FAnimationRecordingSettings::DefaultSampleRate;
+	}
+
+	IntervalTime = 1.0f / SampleRateHz;
+	if (LengthInSeconds == FAnimationRecordingSettings::UnboundedMaximumLength)
+	{
+		// invalid length passed in, default to unbounded
+		MaxFrame = UnBoundedFrameCount;
+	}
+	else
+	{
+		MaxFrame = SampleRateHz * LengthInSeconds;
+	}
+}
+
+bool FAnimationRecorder::SetAnimCompressionScheme(TSubclassOf<UAnimCompress> SchemeClass)
+{
+	if (AnimationObject)
+	{
+		UAnimCompress* const SchemeObject = NewObject<UAnimCompress>(GetTransientPackage(), SchemeClass);
+		if (SchemeObject)
+		{
+			AnimationObject->CompressionScheme = SchemeObject;
+			return true;
+		}
+	}
+
+	return false;
+}
+
+// Internal. Pops up a dialog to get saved asset path
+static bool PromptUserForAssetPath(FString& AssetPath, FString& AssetName)
+{
+	TSharedRef<SCreateAnimationDlg> NewAnimDlg = SNew(SCreateAnimationDlg);
+	if (NewAnimDlg->ShowModal() != EAppReturnType::Cancel)
+	{
+		AssetPath = NewAnimDlg->GetFullAssetPath();
+		AssetName = NewAnimDlg->GetAssetName();
+		return true;
+	}
+
+	return false;
+}
+
+bool FAnimationRecorder::TriggerRecordAnimation(USkeletalMeshComponent* Component)
+{
+	FString AssetPath;
+	FString AssetName;
+
+	if (!Component || !Component->SkeletalMesh || !Component->SkeletalMesh->Skeleton)
+	{
+		return false;
+	}
+
+	// ask for path
+	if (PromptUserForAssetPath(AssetPath, AssetName))
+	{
+		return TriggerRecordAnimation(Component, AssetPath, AssetName);
+	}
+
+	return false;
+}
+
+bool FAnimationRecorder::TriggerRecordAnimation(USkeletalMeshComponent* Component, const FString& InAssetPath, const FString& InAssetName)
+{
+	if (!Component || !Component->SkeletalMesh || !Component->SkeletalMesh->Skeleton)
+	{
+		return false;
+	}
+
+	// create the asset
+	FText InvalidPathReason;
+	bool const bValidPackageName = FPackageName::IsValidLongPackageName(InAssetPath, false, &InvalidPathReason);
+	if (bValidPackageName == false)
+	{
+		UE_LOG(LogAnimation, Log, TEXT("%s is an invalid asset path, prompting user for new asset path. Reason: %s"), *InAssetPath, *InvalidPathReason.ToString());
+	}
+
+	FString ValidatedAssetPath = InAssetPath;
+	FString ValidatedAssetName = InAssetName;
+
+	UObject* Parent = bValidPackageName ? CreatePackage(nullptr, *ValidatedAssetPath) : nullptr;
+	if (Parent == nullptr)
+	{
+		// bad or no path passed in, do the popup
+		if (PromptUserForAssetPath(ValidatedAssetPath, ValidatedAssetName) == false)
+		{
+			return false;
+		}
+
+		Parent = CreatePackage(nullptr, *ValidatedAssetPath);
+	}
+
+	UObject* const Object = LoadObject<UObject>(Parent, *ValidatedAssetName, nullptr, LOAD_None, nullptr);
+	// if object with same name exists, warn user
+	if (Object)
+	{
+		FMessageDialog::Open(EAppMsgType::Ok, NSLOCTEXT("UnrealEd", "Error_AssetExist", "Asset with same name exists. Can't overwrite another asset"));
+		return false;		// failed
+	}
+
+	// If not, create new one now.
+	UAnimSequence* const NewSeq = NewObject<UAnimSequence>(Parent, *ValidatedAssetName, RF_Public | RF_Standalone);
+	if (NewSeq)
+	{
+		// set skeleton
+		NewSeq->SetSkeleton(Component->SkeletalMesh->Skeleton);
+		// Notify the asset registry
+		FAssetRegistryModule::AssetCreated(NewSeq);
+		StartRecord(Component, NewSeq);
+
+		return true;
+	}
+
+	return false;
+}
+
+/** Helper function to get space bases depending on master pose component */
+static void GetBoneTransforms(USkeletalMeshComponent* Component, TArray<FTransform>& BoneTransforms)
+{
+	const USkinnedMeshComponent* const MasterPoseComponentInst = Component->MasterPoseComponent.Get();
+	if(MasterPoseComponentInst)
+	{
+		const TArray<FTransform>& SpaceBases = MasterPoseComponentInst->GetSpaceBases();
+		BoneTransforms.Reset(BoneTransforms.Num());
+		BoneTransforms.AddUninitialized(SpaceBases.Num());
+		for(int32 BoneIndex = 0; BoneIndex < SpaceBases.Num(); BoneIndex++)
+		{
+			if(BoneIndex < Component->GetMasterBoneMap().Num())
+			{
+				int32 MasterBoneIndex = Component->GetMasterBoneMap()[BoneIndex];
+
+				// If ParentBoneIndex is valid, grab matrix from MasterPoseComponent.
+				if(MasterBoneIndex != INDEX_NONE && MasterBoneIndex < SpaceBases.Num())
+				{
+					BoneTransforms[BoneIndex] = SpaceBases[MasterBoneIndex];
+				}
+				else
+				{
+					BoneTransforms[BoneIndex] = FTransform::Identity;
+				}
+			}
+			else
+			{
+				BoneTransforms[BoneIndex] = FTransform::Identity;
+			}
+		}
+	}
+	else
+	{
+		BoneTransforms =  Component->GetSpaceBases();
+	}
+}
+
+void FAnimationRecorder::StartRecord(USkeletalMeshComponent* Component, UAnimSequence* InAnimationObject)
+{
+	TimePassed = 0.f;
+	AnimationObject = InAnimationObject;
+
+	AnimationObject->RawAnimationData.Empty();
+	AnimationObject->TrackToSkeletonMapTable.Empty();
+	AnimationObject->AnimationTrackNames.Empty();
+
+	GetBoneTransforms(Component, PreviousSpacesBases);
+	PreviousAnimCurves = Component->GetAnimationCurves();
+	PreviousComponentToWorld = Component->ComponentToWorld;
+
+	LastFrame = 0;
+	AnimationObject->SequenceLength = 0.f;
+	AnimationObject->NumFrames = 0;
+
+	USkeleton* AnimSkeleton = AnimationObject->GetSkeleton();
+	// add all frames
+	for (int32 BoneIndex=0; BoneIndex <PreviousSpacesBases.Num(); ++BoneIndex)
+	{
+		// verify if this bone exists in skeleton
+		int32 BoneTreeIndex = AnimSkeleton->GetSkeletonBoneIndexFromMeshBoneIndex(Component->SkeletalMesh, BoneIndex);
+		if (BoneTreeIndex != INDEX_NONE)
+		{
+			// add tracks for the bone existing
+			FName BoneTreeName = AnimSkeleton->GetReferenceSkeleton().GetBoneName(BoneTreeIndex);
+			AnimationObject->AnimationTrackNames.Add(BoneTreeName);
+			// add mapping to skeleton bone track
+			AnimationObject->TrackToSkeletonMapTable.Add(FTrackToSkeletonMap(BoneTreeIndex));
+		}
+	}
+
+	int32 NumTracks = AnimationObject->AnimationTrackNames.Num();
+	// add tracks
+	AnimationObject->RawAnimationData.AddZeroed(NumTracks);
+
+	// init notifies
+	AnimationObject->InitializeNotifyTrack();
+
+	// record the first frame
+	Record(Component, PreviousComponentToWorld, PreviousSpacesBases, PreviousAnimCurves,  0);
+}
+
+void FAnimationRecorder::FixupNotifies()
+{
+	if (AnimationObject)
+	{
+		// build notify tracks - first find how many tracks we want
+		for (FAnimNotifyEvent& Event : AnimationObject->Notifies)
+		{
+			if (Event.TrackIndex >= AnimationObject->AnimNotifyTracks.Num())
+			{
+				AnimationObject->AnimNotifyTracks.SetNum(Event.TrackIndex + 1);
+
+				// remake track names to create a nice sequence
+				const int32 TrackNum = AnimationObject->AnimNotifyTracks.Num();
+				for (int32 TrackIndex = 0; TrackIndex < TrackNum; ++TrackIndex)
+				{
+					FAnimNotifyTrack& Track = AnimationObject->AnimNotifyTracks[TrackIndex];
+					Track.TrackName = *FString::FromInt(TrackIndex + 1);
+				}
+			}
+		}
+
+		// now build tracks
+		for (int32 EventIndex = 0; EventIndex < AnimationObject->Notifies.Num(); ++EventIndex)
+		{
+			FAnimNotifyEvent& Event = AnimationObject->Notifies[EventIndex];
+			AnimationObject->AnimNotifyTracks[Event.TrackIndex].Notifies.Add(&AnimationObject->Notifies[EventIndex]);
+		}
+	}
+}
+
+UAnimSequence* FAnimationRecorder::StopRecord(bool bShowMessage)
+{
+	if (AnimationObject)
+	{
+		int32 NumFrames = LastFrame  + 1;
+		AnimationObject->NumFrames = NumFrames;
+
+		// can't use TimePassed. That is just total time that has been passed, not necessarily match with frame count
+		AnimationObject->SequenceLength = (NumFrames>1) ? (NumFrames-1) * IntervalTime : MINIMUM_ANIMATION_LENGTH;
+
+		FixupNotifies();
+
+		// force anim settings for speed, we dont want any fancy recompression at present
+		UAnimationSettings* AnimationSettings = GetMutableDefault<UAnimationSettings>();
+		TSubclassOf<UAnimCompress> OldDefaultCompressionAlgorithm = AnimationSettings->DefaultCompressionAlgorithm;
+		TEnumAsByte<AnimationCompressionFormat> OldRotationCompressionFormat = AnimationSettings->RotationCompressionFormat;
+		TEnumAsByte<AnimationCompressionFormat> OldTranslationCompressionFormat = AnimationSettings->TranslationCompressionFormat;
+
+		AnimationSettings->DefaultCompressionAlgorithm = UAnimCompress_BitwiseCompressOnly::StaticClass();
+		AnimationSettings->RotationCompressionFormat = ACF_None;
+		AnimationSettings->TranslationCompressionFormat = ACF_None;
+
+		// post-process applies compression etc.
+		// @todo figure out why removing redundant keys is inconsistent
+		AnimationObject->RawCurveData.RemoveRedundantKeys();
+		AnimationObject->PostProcessSequence();
+
+		// restore old settings
+		AnimationSettings->DefaultCompressionAlgorithm = OldDefaultCompressionAlgorithm;
+		AnimationSettings->RotationCompressionFormat = OldRotationCompressionFormat;
+		AnimationSettings->TranslationCompressionFormat = OldTranslationCompressionFormat;
+
+		AnimationObject->MarkPackageDirty();
+		
+		// save the package to disk, for convenience and so we can run this in standalone mode
+		if (bAutoSaveAsset)
+		{
+			UPackage* const Package = AnimationObject->GetOutermost();
+			FString const PackageName = Package->GetName();
+			FString const PackageFileName = FPackageName::LongPackageNameToFilename(PackageName, FPackageName::GetAssetPackageExtension());
+
+			UPackage::SavePackage(Package, NULL, RF_Standalone, *PackageFileName, GError, nullptr, false, true, SAVE_NoError);
+		}
+
+		UAnimSequence* ReturnObject = AnimationObject;
+
+		// notify to user
+		if (bShowMessage)
+		{
+			const FText NotificationText = FText::Format(LOCTEXT("RecordAnimation", "'{0}' has been successfully recorded [{1} frames : {2} sec(s) @ {3} Hz]"),
+				FText::FromString(AnimationObject->GetName()),
+				FText::AsNumber(AnimationObject->NumFrames),
+				FText::AsNumber(AnimationObject->SequenceLength),
+				FText::AsNumber(1.f / IntervalTime)
+				);
+					
+			if (GIsEditor)
+			{
+				FNotificationInfo Info(NotificationText);
+				Info.ExpireDuration = 8.0f;
+				Info.bUseLargeFont = false;
+				Info.Hyperlink = FSimpleDelegate::CreateLambda([=]()
+				{
+					TArray<UObject*> Assets;
+					Assets.Add(ReturnObject);
+					FAssetEditorManager::Get().OpenEditorForAssets(Assets);
+				});
+				Info.HyperlinkText = FText::Format(LOCTEXT("OpenNewAnimationHyperlink", "Open {0}"), FText::FromString(AnimationObject->GetName()));
+				TSharedPtr<SNotificationItem> Notification = FSlateNotificationManager::Get().AddNotification(Info);
+				if ( Notification.IsValid() )
+				{
+					Notification->SetCompletionState( SNotificationItem::CS_Success );
+				}
+			}
+
+			FAssetRegistryModule::AssetCreated(AnimationObject);
+		}
+
+		AnimationObject = NULL;
+		PreviousSpacesBases.Empty();
+		PreviousAnimCurves.Empty();
+
+		return ReturnObject;
+	}
+
+	UniqueNotifies.Empty();
+	UniqueNotifyStates.Empty();
+
+	return NULL;
+}
+
+void FAnimationRecorder::UpdateRecord(USkeletalMeshComponent* Component, float DeltaTime)
+{
+	// if no animation object, return
+	if (!AnimationObject || !Component)
+	{
+		return;
+	}
+
+	// no sim time, no record
+	if (DeltaTime <= 0.f)
+	{
+		return;
+	}
+
+	// in-editor we can get a long frame update because of the modal dialog used to pick paths
+	if(DeltaTime > IntervalTime && (LastFrame == 0 || LastFrame == 1))
+	{
+		DeltaTime = IntervalTime;
+	}
+
+	float const PreviousTimePassed = TimePassed;
+	TimePassed += DeltaTime;
+
+	// time passed has been updated
+	// now find what frames we need to update
+	int32 FramesRecorded = LastFrame;
+	int32 FramesToRecord = FPlatformMath::TruncToInt(TimePassed / IntervalTime);
+
+	// notifies need to be done regardless of sample rate
+	if (Component->GetAnimInstance())
+	{
+		RecordNotifies(Component, Component->GetAnimInstance()->NotifyQueue.AnimNotifies, DeltaTime, TimePassed);
+	}
+
+	TArray<FTransform> SpaceBases;
+	GetBoneTransforms(Component, SpaceBases);
+
+	if (FramesRecorded < FramesToRecord)
+	{
+		const FBlendedHeapCurve& AnimCurves = Component->GetAnimationCurves();
+
+		check(SpaceBases.Num() == PreviousSpacesBases.Num());
+
+		TArray<FTransform> BlendedSpaceBases;
+		BlendedSpaceBases.AddZeroed(SpaceBases.Num());
+
+		UE_LOG(LogAnimation, Log, TEXT("DeltaTime : %0.2f, Current Frame Count : %d, Frames To Record : %d, TimePassed : %0.2f"), DeltaTime
+			, FramesRecorded, FramesToRecord, TimePassed);
+
+		// if we need to record frame
+		while (FramesToRecord > FramesRecorded)
+		{
+			// find what frames we need to record
+			// convert to time
+			float CurrentTime = (FramesRecorded + 1) * IntervalTime;
+			float BlendAlpha = (CurrentTime - PreviousTimePassed) / DeltaTime;
+
+			UE_LOG(LogAnimation, Log, TEXT("Current Frame Count : %d, BlendAlpha : %0.2f"), FramesRecorded + 1, BlendAlpha);
+
+			// for now we just concern component space, not skeleton space
+			for (int32 BoneIndex = 0; BoneIndex<SpaceBases.Num(); ++BoneIndex)
+			{
+				BlendedSpaceBases[BoneIndex].Blend(PreviousSpacesBases[BoneIndex], SpaceBases[BoneIndex], BlendAlpha);
+			}
+
+			FTransform BlendedComponentToWorld;
+			BlendedComponentToWorld.Blend(PreviousComponentToWorld, Component->ComponentToWorld, BlendAlpha);
+
+			FBlendedHeapCurve BlendedCurve;
+			if (PreviousAnimCurves.Elements.Num() == AnimCurves.Elements.Num())
+			{
+				BlendedCurve.Blend(PreviousAnimCurves, AnimCurves, BlendAlpha);
+			}
+			else
+			{
+				// just override with AnimCurves for this frames, because UID list has changed
+				// which means new curves are added in run-time
+				BlendedCurve = AnimCurves;
+			}
+
+			Record(Component, BlendedComponentToWorld, BlendedSpaceBases, BlendedCurve, FramesRecorded + 1);
+			++FramesRecorded;
+		}
+	}
+
+	//save to current transform
+	PreviousSpacesBases = SpaceBases;
+	PreviousAnimCurves = Component->GetAnimationCurves();
+	PreviousComponentToWorld = Component->ComponentToWorld;
+
+	// if we passed MaxFrame, just stop it
+	if (MaxFrame != UnBoundedFrameCount && FramesRecorded >= MaxFrame)
+	{
+		UE_LOG(LogAnimation, Log, TEXT("Animation Recording exceeds the time limited (%d mins). Stopping recording animation... "), (int32)((float)MaxFrame / ((1.0f / IntervalTime) * 60.0f)));
+		StopRecord(true);
+	}
+}
+
+void FAnimationRecorder::Record(USkeletalMeshComponent* Component, FTransform const& ComponentToWorld, const TArray<FTransform>& SpacesBases, const FBlendedHeapCurve& AnimationCurves, int32 FrameToAdd)
+{
+	if (ensure(AnimationObject))
+	{
+		USkeleton* AnimSkeleton = AnimationObject->GetSkeleton();
+		for (int32 TrackIndex=0; TrackIndex <AnimationObject->RawAnimationData.Num(); ++TrackIndex)
+		{
+			// verify if this bone exists in skeleton
+			int32 BoneTreeIndex = AnimationObject->GetSkeletonIndexFromTrackIndex(TrackIndex);
+			if (BoneTreeIndex != INDEX_NONE)
+			{
+				int32 BoneIndex = AnimSkeleton->GetMeshBoneIndexFromSkeletonBoneIndex(Component->SkeletalMesh, BoneTreeIndex);
+				int32 ParentIndex = Component->SkeletalMesh->RefSkeleton.GetParentIndex(BoneIndex);
+				FTransform LocalTransform = SpacesBases[BoneIndex];
+				if ( ParentIndex != INDEX_NONE )
+				{
+					LocalTransform.SetToRelativeTransform(SpacesBases[ParentIndex]);
+				}
+				// if record local to world, we'd like to consider component to world to be in root
+				else if (bRecordLocalToWorld)
+				{
+					LocalTransform *= ComponentToWorld;
+				}
+
+				FRawAnimSequenceTrack& RawTrack = AnimationObject->RawAnimationData[TrackIndex];
+				RawTrack.PosKeys.Add(LocalTransform.GetTranslation());
+				RawTrack.RotKeys.Add(LocalTransform.GetRotation());
+				RawTrack.ScaleKeys.Add(LocalTransform.GetScale3D());
+
+				// verification
+				check (FrameToAdd == RawTrack.PosKeys.Num()-1);
+			}
+		}
+
+		const float TimeToRecord = FrameToAdd*IntervalTime;
+		for (int32 CurveIndex = 0; CurveIndex < AnimationCurves.Elements.Num(); ++CurveIndex)
+		{
+			USkeleton::AnimCurveUID UID = (*AnimationCurves.UIDList)[CurveIndex];
+			AnimationObject->RawCurveData.AddFloatCurveKey(UID, AnimationCurves.Elements[CurveIndex].Flags, TimeToRecord, AnimationCurves.Elements[CurveIndex].Value);
+		}
+
+		LastFrame = FrameToAdd;
+	}
+}
+
+void FAnimationRecorder::RecordNotifies(USkeletalMeshComponent* Component, const TArray<const struct FAnimNotifyEvent*>& AnimNotifies, float DeltaTime, float RecordTime)
+{
+	if (ensure(AnimationObject))
+	{
+		// flag notifies as possibly unused this frame
+		for (auto& ActiveNotify : ActiveNotifies)
+		{
+			ActiveNotify.Value = false;
+		}
+
+		int32 AddedThisFrame = 0;
+		for(const FAnimNotifyEvent* NotifyEvent : AnimNotifies)
+		{
+			// we don't want to insert notifies with duration more than once
+			if(NotifyEvent->GetDuration() > 0.0f)
+			{
+				// if this event is active already then don't add it
+				bool bAlreadyActive = false;
+				for(auto& ActiveNotify : ActiveNotifies)
+				{
+					if(NotifyEvent == ActiveNotify.Key)
+					{
+						// flag as active
+						ActiveNotify.Value = true;
+						bAlreadyActive = true;
+						break;
+					}
+				}
+
+				// already active, so skip adding
+				if(bAlreadyActive)
+				{
+					continue;
+				}
+				else
+				{
+					// add a new active notify with duration
+					ActiveNotifies.Add(TPairInitializer<const FAnimNotifyEvent*, bool>(NotifyEvent, true));
+				}
+			}
+
+			// make a new notify from this event & set the current time
+			FAnimNotifyEvent NewEvent = *NotifyEvent;
+			NewEvent.SetTime(RecordTime);
+			NewEvent.TriggerTimeOffset = 0.0f;
+			NewEvent.EndTriggerTimeOffset = 0.0f;
+			
+			// see if we need to create a new notify
+			if(NotifyEvent->Notify)
+			{
+				UAnimNotify** FoundNotify = UniqueNotifies.Find(NotifyEvent->Notify);
+				if(FoundNotify == nullptr)
+				{
+					NewEvent.Notify = Cast<UAnimNotify>(StaticDuplicateObject(NewEvent.Notify, AnimationObject));
+					UniqueNotifies.Add(NotifyEvent->Notify, NewEvent.Notify);
+				}
+				else
+				{
+					NewEvent.Notify = *FoundNotify;
+				}
+			}
+
+			// see if we need to create a new notify state
+			if (NotifyEvent->NotifyStateClass)
+			{
+				UAnimNotifyState** FoundNotifyState = UniqueNotifyStates.Find(NotifyEvent->NotifyStateClass);
+				if (FoundNotifyState == nullptr)
+				{
+					NewEvent.NotifyStateClass = Cast<UAnimNotifyState>(StaticDuplicateObject(NewEvent.NotifyStateClass, AnimationObject));
+					UniqueNotifyStates.Add(NotifyEvent->NotifyStateClass, NewEvent.NotifyStateClass);
+				}
+				else
+				{
+					NewEvent.NotifyStateClass = *FoundNotifyState;
+				}
+			}
+
+			AnimationObject->Notifies.Add(NewEvent);
+			AddedThisFrame++;
+		}
+
+		// remove all notifies that didnt get added this time
+		ActiveNotifies.RemoveAll([](TPair<const FAnimNotifyEvent*, bool>& ActiveNotify){ return !ActiveNotify.Value; });
+
+		UE_LOG(LogAnimation, Log, TEXT("Added notifies : %d"), AddedThisFrame);
+	}
+}
+
+void FAnimationRecorderManager::Tick(float DeltaTime)
+{
+	for (auto& Inst : RecorderInstances)
+	{
+		Inst.Update(DeltaTime);
+	}
+}
+
+void FAnimationRecorderManager::Tick(USkeletalMeshComponent* Component, float DeltaTime)
+{
+	for (auto& Inst : RecorderInstances)
+	{
+		if(Inst.SkelComp == Component)
+		{
+			Inst.Update(DeltaTime);
+		}
+	}
+}
+
+FAnimationRecorderManager::FAnimationRecorderManager()
+{
+}
+
+FAnimationRecorderManager::~FAnimationRecorderManager()
+{
+}
+
+FAnimationRecorderManager& FAnimationRecorderManager::Get()
+{
+	static FAnimationRecorderManager AnimRecorderManager;
+	return AnimRecorderManager;
+}
+
+
+FAnimRecorderInstance::FAnimRecorderInstance()
+	: SkelComp(nullptr)
+	, Recorder(nullptr)
+{
+}
+
+void FAnimRecorderInstance::Init(USkeletalMeshComponent* InComponent, const FString& InAssetPath, const FString& InAssetName, float SampleRateHz, float MaxLength, bool bRecordInWorldSpace, bool bAutoSaveAsset)
+{
+	SkelComp = InComponent;
+	AssetPath = InAssetPath;
+	AssetName = InAssetName;
+	Recorder = MakeShareable(new FAnimationRecorder());
+	Recorder->SetSampleRateAndLength(SampleRateHz, MaxLength);
+	Recorder->bRecordLocalToWorld = bRecordInWorldSpace;
+	Recorder->SetAnimCompressionScheme(UAnimCompress_BitwiseCompressOnly::StaticClass());
+	Recorder->bAutoSaveAsset = bAutoSaveAsset;
+
+	if (InComponent)
+	{
+		CachedSkelCompForcedLodModel = InComponent->ForcedLodModel;
+		InComponent->ForcedLodModel = 1;
+	}
+}
+
+void FAnimRecorderInstance::Init(USkeletalMeshComponent* InComponent, UAnimSequence* InSequence, float SampleRateHz, float MaxLength, bool bRecordInWorldSpace, bool bAutoSaveAsset)
+{
+	SkelComp = InComponent;
+	Sequence = InSequence;
+	Recorder = MakeShareable(new FAnimationRecorder());
+	Recorder->SetSampleRateAndLength(SampleRateHz, MaxLength);
+	Recorder->bRecordLocalToWorld = bRecordInWorldSpace;
+	Recorder->SetAnimCompressionScheme(UAnimCompress_BitwiseCompressOnly::StaticClass());
+	Recorder->bAutoSaveAsset = bAutoSaveAsset;
+
+	if (InComponent)
+	{
+		CachedSkelCompForcedLodModel = InComponent->ForcedLodModel;
+		InComponent->ForcedLodModel = 1;
+	}
+}
+
+FAnimRecorderInstance::~FAnimRecorderInstance()
+{
+}
+
+bool FAnimRecorderInstance::BeginRecording()
+{
+	if (Recorder.IsValid())
+	{
+		if(Sequence.IsValid())
+		{
+			Recorder->StartRecord(SkelComp.Get(), Sequence.Get());
+			return true;
+		}
+		else
+		{
+			return Recorder->TriggerRecordAnimation(SkelComp.Get(), AssetPath, AssetName);
+		}
+	}
+
+	return false;
+}
+
+void FAnimRecorderInstance::Update(float DeltaTime)
+{
+	if (Recorder.IsValid())
+	{
+		Recorder->UpdateRecord(SkelComp.Get(), DeltaTime);
+	}
+}
+
+void FAnimRecorderInstance::FinishRecording(bool bShowMessage)
+{
+	const FText FinishRecordingAnimationSlowTask = LOCTEXT("FinishRecordingAnimationSlowTask", "Finalizing recorded animation");
+	if (Recorder.IsValid())
+	{
+		Recorder->StopRecord(bShowMessage);
+	}
+
+	// restore force lod setting
+	if (SkelComp.IsValid())
+	{
+		SkelComp->ForcedLodModel = CachedSkelCompForcedLodModel;
+	}
+}
+
+bool FAnimationRecorderManager::RecordAnimation(USkeletalMeshComponent* Component, const FString& AssetPath, const FString& AssetName, const FAnimationRecordingSettings& Settings)
+{
+	if (Component)
+	{
+		FAnimRecorderInstance NewInst;
+		NewInst.Init(Component, AssetPath, AssetName, Settings.SampleRate, Settings.Length, Settings.bRecordInWorldSpace, Settings.bAutoSaveAsset);
+		bool const bSuccess = NewInst.BeginRecording();
+		if (bSuccess)
+		{
+			RecorderInstances.Add(NewInst);
+		}
+
+	#if WITH_EDITOR
+			// if recording via PIE, be sure to stop recording cleanly when PIE ends
+			UWorld const* const World = Component->GetWorld();
+			if (World && World->IsPlayInEditor())
+			{
+				FEditorDelegates::EndPIE.AddRaw(this, &FAnimationRecorderManager::HandleEndPIE);
+			}
+	#endif
+
+		return bSuccess;
+	}
+
+	return false;
+}
+
+bool FAnimationRecorderManager::RecordAnimation(USkeletalMeshComponent* Component, UAnimSequence* Sequence, const FAnimationRecordingSettings& Settings)
+{
+	if (Component)
+	{
+		FAnimRecorderInstance NewInst;
+		NewInst.Init(Component, Sequence, Settings.SampleRate, Settings.Length, Settings.bRecordInWorldSpace, Settings.bAutoSaveAsset);
+		bool const bSuccess = NewInst.BeginRecording();
+		if (bSuccess)
+		{
+			RecorderInstances.Add(NewInst);
+		}
+
+	#if WITH_EDITOR
+			// if recording via PIE, be sure to stop recording cleanly when PIE ends
+			UWorld const* const World = Component->GetWorld();
+			if (World && World->IsPlayInEditor())
+			{
+				FEditorDelegates::EndPIE.AddRaw(this, &FAnimationRecorderManager::HandleEndPIE);
+			}
+	#endif
+
+		return bSuccess;
+	}
+
+	return false;
+}
+
+void FAnimationRecorderManager::HandleEndPIE(bool bSimulating)
+{
+	StopRecordingAllAnimations();
+
+#if WITH_EDITOR
+	FEditorDelegates::EndPIE.RemoveAll(this);
+#endif
+}
+
+bool FAnimationRecorderManager::IsRecording(USkeletalMeshComponent* Component)
+{
+	for (FAnimRecorderInstance& Instance : RecorderInstances)
+	{
+		if (Instance.SkelComp == Component)
+		{
+			return Instance.Recorder->InRecording();
+		}
+	}
+
+	return false;
+}
+
+bool FAnimationRecorderManager::IsRecording()
+{
+	for (FAnimRecorderInstance& Instance : RecorderInstances)
+	{
+		if (Instance.Recorder->InRecording())
+		{
+			return true;
+		}
+	}
+
+	return false;
+}
+
+UAnimSequence* FAnimationRecorderManager::GetCurrentlyRecordingSequence(USkeletalMeshComponent* Component)
+{
+	for (FAnimRecorderInstance& Instance : RecorderInstances)
+	{
+		if (Instance.SkelComp == Component)
+		{
+			return Instance.Recorder->GetAnimationObject();
+		}
+	}
+
+	return nullptr;
+}
+
+float FAnimationRecorderManager::GetCurrentRecordingTime(USkeletalMeshComponent* Component)
+{
+	for (FAnimRecorderInstance& Instance : RecorderInstances)
+	{
+		if (Instance.SkelComp == Component)
+		{
+			return Instance.Recorder->GetTimeRecorded();
+		}
+	}
+
+	return 0.0f;
+}
+
+void FAnimationRecorderManager::StopRecordingAnimation(USkeletalMeshComponent* Component, bool bShowMessage)
+{
+	for (int32 Idx = 0; Idx < RecorderInstances.Num(); ++Idx)
+	{
+		FAnimRecorderInstance& Inst = RecorderInstances[Idx];
+		if (Inst.SkelComp == Component)
+		{
+			// stop and finalize recoded data
+			Inst.FinishRecording(bShowMessage);
+
+			// remove instance, which will clean itself up
+			RecorderInstances.RemoveAtSwap(Idx);
+
+			// all done
+			break;
+		}
+	}
+}
+
+void FAnimationRecorderManager::StopRecordingDeadAnimations(bool bShowMessage)
+{
+	RecorderInstances.RemoveAll([&](FAnimRecorderInstance& Instance)
+		{
+			if (!Instance.SkelComp.IsValid())
+			{
+				// stop and finalize recorded data
+				Instance.FinishRecording(bShowMessage);
+
+				// make sure we are cleaned up
+				return true;
+			}
+
+			return false;
+		}
+	);
+}
+
+void FAnimationRecorderManager::StopRecordingAllAnimations()
+{
+	for (auto& Inst : RecorderInstances)
+	{
+		Inst.FinishRecording();
+	}
+
+	RecorderInstances.Empty();
+}
+
+#undef LOCTEXT_NAMESPACE
+
diff --git a/Engine/Source/Editor/SequenceRecorder/Private/AnimationRecorder.h b/Engine/Source/Editor/SequenceRecorder/Private/AnimationRecorder.h
new file mode 100644
--- /dev/null
+++ b/Engine/Source/Editor/SequenceRecorder/Private/AnimationRecorder.h
@@ -0,0 +1,138 @@
+// Copyright 1998-2016 Epic Games, Inc. All Rights Reserved.
+
+#pragma once
+
+//////////////////////////////////////////////////////////////////////////
+// FAnimationRecorder
+
+// records the mesh pose to animation input
+struct FAnimationRecorder
+{
+private:
+	/** Frame count used to signal an unbounded animation */
+	static const int32 UnBoundedFrameCount = -1;
+
+private:
+	float IntervalTime;
+	int32 MaxFrame;
+	int32 LastFrame;
+	float TimePassed;
+	UAnimSequence* AnimationObject;
+	TArray<FTransform> PreviousSpacesBases;
+	FBlendedHeapCurve PreviousAnimCurves;
+	FTransform PreviousComponentToWorld;
+
+	/** Array of currently active notifies that have duration */
+	TArray<TPair<const FAnimNotifyEvent*, bool>> ActiveNotifies;
+
+	/** Unique notifies added to this sequence during recording */
+	TMap<UAnimNotify*, UAnimNotify*> UniqueNotifies;
+
+	/** Unique notify states added to this sequence during recording */
+	TMap<UAnimNotifyState*, UAnimNotifyState*> UniqueNotifyStates;
+
+	static float DefaultSampleRate;
+
+public:
+	FAnimationRecorder();
+	~FAnimationRecorder();
+
+	/** Starts recording an animation. Prompts for asset path and name via dialog if none provided */
+	bool TriggerRecordAnimation(USkeletalMeshComponent* Component);
+	bool TriggerRecordAnimation(USkeletalMeshComponent* Component, const FString& AssetPath, const FString& AssetName);
+
+	void StartRecord(USkeletalMeshComponent* Component, UAnimSequence* InAnimationObject);
+	UAnimSequence* StopRecord(bool bShowMessage);
+	void UpdateRecord(USkeletalMeshComponent* Component, float DeltaTime);
+	UAnimSequence* GetAnimationObject() const { return AnimationObject; }
+	bool InRecording() const { return AnimationObject != nullptr; }
+	float GetTimeRecorded() const { return TimePassed; }
+
+	/** Sets a new sample rate & max length for this recorder. Don't call while recording. */
+	void SetSampleRateAndLength(float SampleRateHz, float LengthInMinutes);
+
+	bool SetAnimCompressionScheme(TSubclassOf<class UAnimCompress> SchemeClass);
+
+	/** If true, it will record root to include LocalToWorld */
+	uint8 bRecordLocalToWorld :1;
+	/** If true, asset will be saved to disk after recording. If false, asset will remain in mem and can be manually saved. */
+	uint8 bAutoSaveAsset : 1;
+
+private:
+	void Record(USkeletalMeshComponent* Component, FTransform const& ComponentToWorld, const TArray<FTransform>& SpacesBases, const FBlendedHeapCurve& AnimationCurves, int32 FrameToAdd);
+
+	void RecordNotifies(USkeletalMeshComponent* Component, const TArray<const struct FAnimNotifyEvent*>& AnimNotifies, float DeltaTime, float RecordTime);
+
+	void FixupNotifies();
+};
+
+//////////////////////////////////////////////////////////////////////////
+// FAnimRecorderInstance
+
+struct FAnimRecorderInstance
+{
+public:
+	FAnimRecorderInstance();
+	~FAnimRecorderInstance();
+
+	void Init(USkeletalMeshComponent* InComponent, const FString& InAssetPath, const FString& InAssetName, float SampleRateHz, float MaxLength, bool bRecordInWorldSpace, bool bAutoSaveAsset);
+
+	void Init(USkeletalMeshComponent* InComponent, UAnimSequence* InSequence, float SampleRateHz, float MaxLength, bool bRecordInWorldSpace, bool bAutoSaveAsset);
+
+	bool BeginRecording();
+	void Update(float DeltaTime);
+	void FinishRecording(bool bShowMessage = true);
+
+	TWeakObjectPtr<USkeletalMeshComponent> SkelComp;
+	TWeakObjectPtr<UAnimSequence> Sequence;
+	FString AssetPath;
+	FString AssetName;
+
+	/** Original ForcedLodModel setting on the SkelComp, so we can modify it and restore it when we are done. */
+	int CachedSkelCompForcedLodModel;
+
+	TSharedPtr<FAnimationRecorder> Recorder;
+};
+
+
+//////////////////////////////////////////////////////////////////////////
+// FAnimationRecorderManager
+
+struct FAnimationRecorderManager
+{
+public:
+	/** Singleton accessor */
+	static FAnimationRecorderManager& Get();
+
+	/** Destructor */
+	virtual ~FAnimationRecorderManager();
+
+	/** Starts recording an animation. */
+	bool RecordAnimation(USkeletalMeshComponent* Component, const FString& AssetPath = FString(), const FString& AssetName = FString(), const FAnimationRecordingSettings& Settings = FAnimationRecordingSettings());
+
+	bool RecordAnimation(USkeletalMeshComponent* Component, UAnimSequence* Sequence, const FAnimationRecordingSettings& Settings = FAnimationRecordingSettings());
+
+	bool IsRecording(USkeletalMeshComponent* Component);
+
+	bool IsRecording();
+
+	UAnimSequence* GetCurrentlyRecordingSequence(USkeletalMeshComponent* Component);
+	float GetCurrentRecordingTime(USkeletalMeshComponent* Component);
+	void StopRecordingAnimation(USkeletalMeshComponent* Component, bool bShowMessage = true);
+	void StopRecordingAllAnimations();
+
+	void Tick(float DeltaTime);
+
+	void Tick(USkeletalMeshComponent* Component, float DeltaTime);
+
+	void StopRecordingDeadAnimations(bool bShowMessage = true);
+
+private:
+	/** Constructor, private - use Get() function */
+	FAnimationRecorderManager();
+
+	TArray<FAnimRecorderInstance> RecorderInstances;
+
+	void HandleEndPIE(bool bSimulating);
+};
+
diff --git a/Engine/Source/Editor/SequenceRecorder/Private/IMovieScenePropertyRecorder.h b/Engine/Source/Editor/SequenceRecorder/Private/IMovieScenePropertyRecorder.h
new file mode 100644
--- /dev/null
+++ b/Engine/Source/Editor/SequenceRecorder/Private/IMovieScenePropertyRecorder.h
@@ -0,0 +1,12 @@
+// Copyright 1998-2015 Epic Games, Inc. All Rights Reserved.
+
+#pragma once
+
+class IMovieScenePropertyRecorder
+{
+public:
+	virtual void CreateSection(UObject* InObjectToRecord, class UMovieScene* MovieScene, const FGuid& Guid, float Time, bool bRecord) = 0;
+	virtual void FinalizeSection() = 0;
+	virtual void Record(float CurrentTime) = 0;
+	virtual void InvalidateObjectToRecord() = 0;
+};
\ No newline at end of file
diff --git a/Engine/Source/Editor/SequenceRecorder/Private/PropertyRecorders/MovieScene3DAttachPropertyRecorder.cpp b/Engine/Source/Editor/SequenceRecorder/Private/PropertyRecorders/MovieScene3DAttachPropertyRecorder.cpp
new file mode 100644
--- /dev/null
+++ b/Engine/Source/Editor/SequenceRecorder/Private/PropertyRecorders/MovieScene3DAttachPropertyRecorder.cpp
@@ -0,0 +1,71 @@
+// Copyright 1998-2015 Epic Games, Inc. All Rights Reserved.
+
+#include "SequenceRecorderPrivatePCH.h"
+#include "MovieScene3DAttachPropertyRecorder.h"
+#include "MovieScene3DAttachSection.h"
+#include "MovieScene3DAttachTrack.h"
+#include "MovieScene.h"
+#include "SequenceRecorderUtils.h"
+#include "SequenceRecorder.h"
+#include "ActorRecording.h"
+
+void FMovieScene3DAttachPropertyRecorder::CreateSection(UObject* InObjectToRecord, UMovieScene* InMovieScene, const FGuid& Guid, float Time, bool bRecord)
+{
+	ObjectGuid = Guid;
+	ActorToRecord = CastChecked<AActor>(InObjectToRecord);
+	MovieScene = InMovieScene;
+
+	bRecording = bRecord;
+
+	Record(Time);
+}
+
+void FMovieScene3DAttachPropertyRecorder::FinalizeSection()
+{
+	bRecording = false;
+}
+
+void FMovieScene3DAttachPropertyRecorder::Record(float CurrentTime)
+{
+	if(ActorToRecord.IsValid())
+	{
+		if(bRecording)
+		{
+			if(MovieSceneSection.IsValid())
+			{
+				MovieSceneSection->SetEndTime(CurrentTime);
+			}
+
+			// get attachment and check if the actor we are attached to is being recorded
+			FName SocketName;
+			FName ComponentName;
+			AActor* AttachedToActor = SequenceRecorderUtils::GetAttachment(ActorToRecord.Get(), SocketName, ComponentName);
+			UActorRecording* ActorRecording = FSequenceRecorder::Get().FindRecording(AttachedToActor);
+			if(AttachedToActor && ActorRecording)
+			{
+				// create the track if we haven't already
+				if(!AttachTrack.IsValid())
+				{
+					AttachTrack = MovieScene->AddTrack<UMovieScene3DAttachTrack>(ObjectGuid);
+				}
+
+				// check if we need a section or if the actor we are attached to has changed
+				if(!MovieSceneSection.IsValid() || AttachedToActor != ActorAttachedTo.Get())
+				{
+					MovieSceneSection = Cast<UMovieScene3DAttachSection>(AttachTrack->CreateNewSection());
+					MovieSceneSection->SetStartTime(CurrentTime);
+					MovieSceneSection->SetConstraintId(ActorRecording->Guid);
+					MovieSceneSection->AttachSocketName = SocketName;
+					MovieSceneSection->AttachComponentName = ComponentName;
+				}
+
+				ActorAttachedTo = AttachedToActor;
+			}
+			else
+			{
+				// no attachment, so end the section recording if we have any
+				MovieSceneSection = nullptr;
+			}
+		}
+	}
+}
\ No newline at end of file
diff --git a/Engine/Source/Editor/SequenceRecorder/Private/PropertyRecorders/MovieScene3DAttachPropertyRecorder.h b/Engine/Source/Editor/SequenceRecorder/Private/PropertyRecorders/MovieScene3DAttachPropertyRecorder.h
new file mode 100644
--- /dev/null
+++ b/Engine/Source/Editor/SequenceRecorder/Private/PropertyRecorders/MovieScene3DAttachPropertyRecorder.h
@@ -0,0 +1,41 @@
+// Copyright 1998-2015 Epic Games, Inc. All Rights Reserved.
+
+#pragma once
+
+#include "IMovieScenePropertyRecorder.h"
+
+class FMovieScene3DAttachPropertyRecorder : public IMovieScenePropertyRecorder
+{
+public:
+	virtual ~FMovieScene3DAttachPropertyRecorder() {}
+
+	virtual void CreateSection(UObject* InObjectToRecord, class UMovieScene* InMovieScene, const FGuid& Guid, float Time, bool bRecord) override;
+	virtual void FinalizeSection() override;
+	virtual void Record(float CurrentTime) override;
+	virtual void InvalidateObjectToRecord() override
+	{
+		ActorToRecord = nullptr;
+	}
+
+private:
+	/** Object to record from */
+	TLazyObjectPtr<class AActor> ActorToRecord;
+
+	/** Section to record to */
+	TWeakObjectPtr<class UMovieScene3DAttachSection> MovieSceneSection;
+
+	/** Track we are recording to */
+	TWeakObjectPtr<class UMovieScene3DAttachTrack> AttachTrack;
+
+	/** Movie scene we are recording to */
+	TWeakObjectPtr<class UMovieScene> MovieScene;
+
+	/** Track the actor we are attached to */
+	TLazyObjectPtr<class AActor> ActorAttachedTo;
+
+	/** Identifier of the object we are recording */
+	FGuid ObjectGuid;
+
+	/** Flag if we are actually recording or not */
+	bool bRecording;
+};
\ No newline at end of file
diff --git a/Engine/Source/Editor/SequenceRecorder/Private/PropertyRecorders/MovieScene3DTransformPropertyRecorder.cpp b/Engine/Source/Editor/SequenceRecorder/Private/PropertyRecorders/MovieScene3DTransformPropertyRecorder.cpp
new file mode 100644
--- /dev/null
+++ b/Engine/Source/Editor/SequenceRecorder/Private/PropertyRecorders/MovieScene3DTransformPropertyRecorder.cpp
@@ -0,0 +1,271 @@
+// Copyright 1998-2015 Epic Games, Inc. All Rights Reserved.
+
+#include "SequenceRecorderPrivatePCH.h"
+#include "MovieScene3DTransformPropertyRecorder.h"
+#include "MovieScene.h"
+#include "MovieScene3DTransformSection.h"
+#include "MovieScene3DTransformTrack.h"
+#include "MovieSceneAnimationPropertyRecorder.h"
+#include "SequenceRecorderUtils.h"
+
+void FMovieScene3DTransformPropertyRecorder::CreateSection(UObject* InObjectToRecord, UMovieScene* InMovieScene, const FGuid& Guid, float Time, bool bRecord) 
+{
+	ObjectToRecord = InObjectToRecord;
+	bWasAttached = false;
+		
+	MovieScene = InMovieScene;
+	MovieSceneTrack = InMovieScene->AddTrack<UMovieScene3DTransformTrack>(Guid);
+	if(MovieSceneTrack.IsValid())
+	{
+		MovieSceneSection = Cast<UMovieScene3DTransformSection>(MovieSceneTrack->CreateNewSection());
+
+		MovieSceneTrack->AddSection(*MovieSceneSection);
+
+		const bool bUnwindRotation = false;
+
+		DefaultTransform = GetTransformToRecord();
+		FVector EulerRotation = DefaultTransform.GetRotation().Rotator().Euler();
+
+		MovieSceneSection->SetDefault(FTransformKey(EKey3DTransformChannel::Translation, EAxis::X, DefaultTransform.GetTranslation().X, bUnwindRotation));
+		MovieSceneSection->SetDefault(FTransformKey(EKey3DTransformChannel::Translation, EAxis::Y, DefaultTransform.GetTranslation().Y, bUnwindRotation));
+		MovieSceneSection->SetDefault(FTransformKey(EKey3DTransformChannel::Translation, EAxis::Z, DefaultTransform.GetTranslation().Z, bUnwindRotation));
+
+		MovieSceneSection->SetDefault(FTransformKey(EKey3DTransformChannel::Rotation, EAxis::X, EulerRotation.X, bUnwindRotation));
+		MovieSceneSection->SetDefault(FTransformKey(EKey3DTransformChannel::Rotation, EAxis::Y, EulerRotation.Y, bUnwindRotation));
+		MovieSceneSection->SetDefault(FTransformKey(EKey3DTransformChannel::Rotation, EAxis::Z, EulerRotation.Z, bUnwindRotation));
+
+		MovieSceneSection->SetDefault(FTransformKey(EKey3DTransformChannel::Scale, EAxis::X, DefaultTransform.GetScale3D().X, bUnwindRotation));
+		MovieSceneSection->SetDefault(FTransformKey(EKey3DTransformChannel::Scale, EAxis::Y, DefaultTransform.GetScale3D().Y, bUnwindRotation));
+		MovieSceneSection->SetDefault(FTransformKey(EKey3DTransformChannel::Scale, EAxis::Z, DefaultTransform.GetScale3D().Z, bUnwindRotation));
+
+		MovieSceneSection->SetStartTime(Time);
+		MovieSceneSection->SetIsInfinite(true);
+	}
+
+	bRecording = bRecord;
+
+	Record(Time);
+}
+
+void FMovieScene3DTransformPropertyRecorder::FinalizeSection()
+{
+	FScopedSlowTask SlowTask(4.0f, NSLOCTEXT("SequenceRecorder", "ProcessingTransforms", "Processing Transforms"));
+
+	bRecording = false;
+
+	// if we have a valid animation recorder, we need to build our transforms from the animation
+	// so we properly synchronize our keyframes
+	if(AnimRecorder.IsValid())
+	{
+		check(BufferedTransforms.Num() == 0);
+
+		UAnimSequence* AnimSequence = AnimRecorder->GetSequence();
+		USkeletalMesh* SkeletalMesh = AnimRecorder->GetSkeletalMesh();
+		if(AnimSequence && SkeletalMesh)
+		{
+			// find the root bone
+			int32 RootIndex = INDEX_NONE;
+			USkeleton* AnimSkeleton = AnimSequence->GetSkeleton();
+			for (int32 TrackIndex = 0; TrackIndex < AnimSequence->RawAnimationData.Num(); ++TrackIndex)
+			{
+				// verify if this bone exists in skeleton
+				int32 BoneTreeIndex = AnimSequence->GetSkeletonIndexFromTrackIndex(TrackIndex);
+				if (BoneTreeIndex != INDEX_NONE)
+				{
+					int32 BoneIndex = AnimSkeleton->GetMeshBoneIndexFromSkeletonBoneIndex(SkeletalMesh, BoneTreeIndex);
+					int32 ParentIndex = SkeletalMesh->RefSkeleton.GetParentIndex(BoneIndex);
+					if(ParentIndex == INDEX_NONE)
+					{
+						// found root
+						RootIndex = BoneIndex;
+						break;
+					}
+				}
+			}
+
+			check(RootIndex != INDEX_NONE);
+
+			const float StartTime = MovieSceneSection->GetStartTime();
+
+			// we may need to offset the transform here if the animation was not recorded on the root component
+			FTransform InvComponentTransform = AnimRecorder->GetComponentTransform().Inverse();
+
+			FRawAnimSequenceTrack& RawTrack = AnimSequence->RawAnimationData[RootIndex];
+			const int32 KeyCount = FMath::Max(FMath::Max(RawTrack.PosKeys.Num(), RawTrack.RotKeys.Num()), RawTrack.ScaleKeys.Num());
+			for(int32 KeyIndex = 0; KeyIndex < KeyCount; KeyIndex++)
+			{
+				FTransform Transform;
+				if(RawTrack.PosKeys.IsValidIndex(KeyIndex))
+				{
+					Transform.SetTranslation(RawTrack.PosKeys[KeyIndex]);
+				}
+				else if(RawTrack.PosKeys.Num() > 0)
+				{
+					Transform.SetTranslation(RawTrack.PosKeys[0]);
+				}
+				
+				if(RawTrack.RotKeys.IsValidIndex(KeyIndex))
+				{
+					Transform.SetRotation(RawTrack.RotKeys[KeyIndex]);
+				}
+				else if(RawTrack.RotKeys.Num() > 0)
+				{
+					Transform.SetRotation(RawTrack.RotKeys[0]);
+				}
+
+				if(RawTrack.ScaleKeys.IsValidIndex(KeyIndex))
+				{
+					Transform.SetScale3D(RawTrack.ScaleKeys[KeyIndex]);
+				}
+				else if(RawTrack.ScaleKeys.Num() > 0)
+				{
+					Transform.SetScale3D(RawTrack.ScaleKeys[0]);
+				}
+
+				BufferedTransforms.Add(FBufferedTransformKey(InvComponentTransform * Transform, StartTime + AnimSequence->GetTimeAtFrame(KeyIndex)));
+			}
+		}
+	}
+
+	SlowTask.EnterProgressFrame();
+
+	// Try to 're-wind' rotations that look like axis flips
+	// We need to do this as a post-process because the recorder cant reliably access 'wound' rotations:
+	// - Net quantize may use quaternions.
+	// - Scene components cache transforms as quaternions.
+	// - Gameplay is free to clamp/fmod rotations as it sees fit.
+	int32 TransformCount = BufferedTransforms.Num();
+	for(int32 TransformIndex = 0; TransformIndex < TransformCount - 1; TransformIndex++)
+	{
+		FRotator& Rotator = BufferedTransforms[TransformIndex].WoundRotation;
+		FRotator& NextRotator = BufferedTransforms[TransformIndex + 1].WoundRotation;
+
+		FMath::WindRelativeAnglesDegrees(Rotator.Pitch, NextRotator.Pitch);
+		FMath::WindRelativeAnglesDegrees(Rotator.Yaw, NextRotator.Yaw);
+		FMath::WindRelativeAnglesDegrees(Rotator.Roll, NextRotator.Roll);
+	}
+
+	SlowTask.EnterProgressFrame();
+
+	// never unwind rotations
+	const bool bUnwindRotation = false;
+	// If we are syncing to an animation, use linear interpolation to avoid foot sliding etc. 
+	// Otherwise use cubic for better quality (much better for projectiles etc.)
+	const EMovieSceneKeyInterpolation Interpolation = AnimRecorder.IsValid() ? EMovieSceneKeyInterpolation::Linear : EMovieSceneKeyInterpolation::Auto;
+
+	// add buffered transforms
+	for(const FBufferedTransformKey& BufferedTransform : BufferedTransforms)
+	{
+		const FVector Translation = BufferedTransform.Transform.GetTranslation();
+		const FVector Rotation = BufferedTransform.WoundRotation.Euler();
+		const FVector Scale = BufferedTransform.Transform.GetScale3D();
+
+		MovieSceneSection->AddKey(BufferedTransform.KeyTime, FTransformKey(EKey3DTransformChannel::Translation, EAxis::X, Translation.X, bUnwindRotation), Interpolation);
+		MovieSceneSection->AddKey(BufferedTransform.KeyTime, FTransformKey(EKey3DTransformChannel::Translation, EAxis::Y, Translation.Y, bUnwindRotation), Interpolation);
+		MovieSceneSection->AddKey(BufferedTransform.KeyTime, FTransformKey(EKey3DTransformChannel::Translation, EAxis::Z, Translation.Z, bUnwindRotation), Interpolation);
+
+		MovieSceneSection->AddKey(BufferedTransform.KeyTime, FTransformKey(EKey3DTransformChannel::Rotation, EAxis::X, Rotation.X, bUnwindRotation), Interpolation);
+		MovieSceneSection->AddKey(BufferedTransform.KeyTime, FTransformKey(EKey3DTransformChannel::Rotation, EAxis::Y, Rotation.Y, bUnwindRotation), Interpolation);
+		MovieSceneSection->AddKey(BufferedTransform.KeyTime, FTransformKey(EKey3DTransformChannel::Rotation, EAxis::Z, Rotation.Z, bUnwindRotation), Interpolation);
+
+		MovieSceneSection->AddKey(BufferedTransform.KeyTime, FTransformKey(EKey3DTransformChannel::Scale, EAxis::X, Scale.X, bUnwindRotation), Interpolation);
+		MovieSceneSection->AddKey(BufferedTransform.KeyTime, FTransformKey(EKey3DTransformChannel::Scale, EAxis::Y, Scale.Y, bUnwindRotation), Interpolation);
+		MovieSceneSection->AddKey(BufferedTransform.KeyTime, FTransformKey(EKey3DTransformChannel::Scale, EAxis::Z, Scale.Z, bUnwindRotation), Interpolation);
+	}
+
+	BufferedTransforms.Empty();
+
+	SlowTask.EnterProgressFrame();
+
+	// now remove linear keys
+	TPair<FRichCurve*, float> CurvesAndTolerances[] =
+	{
+		TPairInitializer<FRichCurve*, float>(&MovieSceneSection->GetTranslationCurve(EAxis::X), KINDA_SMALL_NUMBER),
+		TPairInitializer<FRichCurve*, float>(&MovieSceneSection->GetTranslationCurve(EAxis::Y), KINDA_SMALL_NUMBER),
+		TPairInitializer<FRichCurve*, float>(&MovieSceneSection->GetTranslationCurve(EAxis::Z), KINDA_SMALL_NUMBER),
+		TPairInitializer<FRichCurve*, float>(&MovieSceneSection->GetRotationCurve(EAxis::X), KINDA_SMALL_NUMBER),
+		TPairInitializer<FRichCurve*, float>(&MovieSceneSection->GetRotationCurve(EAxis::Y), KINDA_SMALL_NUMBER),
+		TPairInitializer<FRichCurve*, float>(&MovieSceneSection->GetRotationCurve(EAxis::Z), KINDA_SMALL_NUMBER),
+		TPairInitializer<FRichCurve*, float>(&MovieSceneSection->GetScaleCurve(EAxis::X), KINDA_SMALL_NUMBER),
+		TPairInitializer<FRichCurve*, float>(&MovieSceneSection->GetScaleCurve(EAxis::Y), KINDA_SMALL_NUMBER),
+		TPairInitializer<FRichCurve*, float>(&MovieSceneSection->GetScaleCurve(EAxis::Z), KINDA_SMALL_NUMBER),
+	};
+
+	for(TPair<FRichCurve*, float>& CurveAndTolerance : CurvesAndTolerances)
+	{
+		CurveAndTolerance.Key->RemoveRedundantKeys(CurveAndTolerance.Value);
+	}
+
+	// we cant remove redundant tracks if we were attached as the playback relies on update order of
+	// transform tracks. Without this track, relative transforms would accumulate.
+	if(!bWasAttached)
+	{
+		// now we have reduced our keys, if we dont have any, remove the section as it is redundant
+		if( MovieSceneSection->GetTranslationCurve(EAxis::X).Keys.Num() == 0 && 
+			MovieSceneSection->GetTranslationCurve(EAxis::Y).Keys.Num() == 0 && 
+			MovieSceneSection->GetTranslationCurve(EAxis::Z).Keys.Num() == 0 && 
+			MovieSceneSection->GetRotationCurve(EAxis::X).Keys.Num() == 0 && 
+			MovieSceneSection->GetRotationCurve(EAxis::Y).Keys.Num() == 0 && 
+			MovieSceneSection->GetRotationCurve(EAxis::Z).Keys.Num() == 0 && 
+			MovieSceneSection->GetScaleCurve(EAxis::X).Keys.Num() == 0 && 
+			MovieSceneSection->GetScaleCurve(EAxis::Y).Keys.Num() == 0 && 
+			MovieSceneSection->GetScaleCurve(EAxis::Z).Keys.Num() == 0)
+		{
+			if(DefaultTransform.Equals(FTransform::Identity))
+			{
+				MovieScene->RemoveTrack(*MovieSceneTrack.Get());
+			}
+		}
+	}
+
+	SlowTask.EnterProgressFrame();
+}
+
+void FMovieScene3DTransformPropertyRecorder::Record(float CurrentTime)
+{
+	if(ObjectToRecord.IsValid())
+	{
+		if(USceneComponent* SceneComponent = Cast<USceneComponent>(ObjectToRecord.Get()))
+		{
+			// dont record non-registered scene components
+			if(!SceneComponent->IsRegistered())
+			{
+				return;
+			}
+		}
+
+		MovieSceneSection->SetEndTime(CurrentTime);
+
+		if(bRecording)
+		{
+			// don't record from the transform of the component/actor if we are synchronizing with an animation
+			if(!AnimRecorder.IsValid())
+			{
+				BufferedTransforms.Add(FBufferedTransformKey(GetTransformToRecord(), CurrentTime));
+			}
+		}
+	}
+}
+
+FTransform FMovieScene3DTransformPropertyRecorder::GetTransformToRecord()
+{
+	if(USceneComponent* SceneComponent = Cast<USceneComponent>(ObjectToRecord.Get()))
+	{
+		return SceneComponent->GetRelativeTransform();
+	}
+	else if(AActor* Actor = Cast<AActor>(ObjectToRecord.Get()))
+	{
+		FName SocketName, ComponentName;
+		if(SequenceRecorderUtils::GetAttachment(Actor, SocketName, ComponentName) != nullptr && Actor->GetRootComponent() != nullptr)
+		{
+			bWasAttached = true;
+			return Actor->GetRootComponent()->GetRelativeTransform();
+		}
+		else
+		{
+			return Actor->GetTransform();
+		}
+	}
+
+	return FTransform::Identity;
+}
\ No newline at end of file
diff --git a/Engine/Source/Editor/SequenceRecorder/Private/PropertyRecorders/MovieScene3DTransformPropertyRecorder.h b/Engine/Source/Editor/SequenceRecorder/Private/PropertyRecorders/MovieScene3DTransformPropertyRecorder.h
new file mode 100644
--- /dev/null
+++ b/Engine/Source/Editor/SequenceRecorder/Private/PropertyRecorders/MovieScene3DTransformPropertyRecorder.h
@@ -0,0 +1,69 @@
+// Copyright 1998-2015 Epic Games, Inc. All Rights Reserved.
+
+#pragma once
+
+#include "IMovieScenePropertyRecorder.h"
+
+/** Structure used to buffer up transform keys. Keys are inserted into tracks in FinalizeSection() */
+struct FBufferedTransformKey
+{
+	FBufferedTransformKey(const FTransform& InTransform, float InKeyTime)
+		: Transform(InTransform)
+		, WoundRotation(InTransform.Rotator())
+		, KeyTime(InKeyTime)
+	{
+	}
+
+	FTransform Transform;
+	FRotator WoundRotation;
+	float KeyTime;
+};
+
+class FMovieScene3DTransformPropertyRecorder : public IMovieScenePropertyRecorder
+{
+public:
+	FMovieScene3DTransformPropertyRecorder(TSharedPtr<class FMovieSceneAnimationPropertyRecorder> InAnimRecorder = nullptr) 
+		: AnimRecorder(InAnimRecorder)
+	{}
+
+	virtual ~FMovieScene3DTransformPropertyRecorder() {}
+
+	virtual void CreateSection(UObject* InObjectToRecord, class UMovieScene* InMovieScene, const FGuid& Guid, float Time, bool bRecord) override;
+	virtual void FinalizeSection() override;
+	virtual void Record(float CurrentTime) override;
+	virtual void InvalidateObjectToRecord() override
+	{
+		ObjectToRecord = nullptr;
+	}
+
+private:
+	FTransform GetTransformToRecord();
+
+private:
+	/** Object to record from */
+	TLazyObjectPtr<UObject> ObjectToRecord;
+
+	/** MovieScene to record to */
+	TWeakObjectPtr<class UMovieScene> MovieScene;
+
+	/** Track to record to */
+	TWeakObjectPtr<class UMovieScene3DTransformTrack> MovieSceneTrack;
+
+	/** Section to record to */
+	TWeakObjectPtr<class UMovieScene3DTransformSection> MovieSceneSection;
+
+	/** Buffer of transform keys. Keys are inserted into tracks in FinalizeSection() */
+	TArray<FBufferedTransformKey> BufferedTransforms;
+
+	/** Flag if we are actually recording or not */
+	bool bRecording;
+
+	/** Animation recorder we use to sync our transforms to */
+	TSharedPtr<class FMovieSceneAnimationPropertyRecorder> AnimRecorder;
+
+	/** The default transform this recording starts with */
+	FTransform DefaultTransform;
+
+	/** Flag indicating that some time while this recorder was active an attachment was also in place */
+	bool bWasAttached;
+};
\ No newline at end of file
diff --git a/Engine/Source/Editor/SequenceRecorder/Private/PropertyRecorders/MovieSceneAnimationPropertyRecorder.cpp b/Engine/Source/Editor/SequenceRecorder/Private/PropertyRecorders/MovieSceneAnimationPropertyRecorder.cpp
new file mode 100644
--- /dev/null
+++ b/Engine/Source/Editor/SequenceRecorder/Private/PropertyRecorders/MovieSceneAnimationPropertyRecorder.cpp
@@ -0,0 +1,133 @@
+// Copyright 1998-2015 Epic Games, Inc. All Rights Reserved.
+
+#include "SequenceRecorderPrivatePCH.h"
+#include "MovieSceneAnimationPropertyRecorder.h"
+#include "MovieSceneSkeletalAnimationTrack.h"
+#include "MovieSceneSkeletalAnimationSection.h"
+#include "MovieScene.h"
+#include "SequenceRecorderUtils.h"
+#include "SequenceRecorderSettings.h"
+#include "Runtime/AssetRegistry/Public/AssetRegistryModule.h"
+
+FMovieSceneAnimationPropertyRecorder::FMovieSceneAnimationPropertyRecorder(FAnimationRecordingSettings& InAnimationSettings, UAnimSequence* InSpecifiedSequence)
+	: Sequence(InSpecifiedSequence)
+	, AnimationSettings(InAnimationSettings)
+{
+}
+
+void FMovieSceneAnimationPropertyRecorder::CreateSection(UObject* InObjectToRecord, UMovieScene* MovieScene, const FGuid& Guid, float Time, bool bRecord)
+{
+	ObjectToRecord = InObjectToRecord;
+
+	AActor* Actor = nullptr;
+	SkeletalMeshComponent = Cast<USkeletalMeshComponent>(InObjectToRecord);
+	if(!SkeletalMeshComponent.IsValid())
+	{
+		Actor = Cast<AActor>(InObjectToRecord);
+		if(Actor != nullptr)
+		{
+			SkeletalMeshComponent = Actor->FindComponentByClass<USkeletalMeshComponent>();
+		}
+	}
+	else
+	{
+		Actor = SkeletalMeshComponent->GetOwner();
+	}
+
+	if(SkeletalMeshComponent.IsValid())
+	{
+		SkeletalMesh = SkeletalMeshComponent->SkeletalMesh;
+
+		// turn off URO and make sure we always update even if out of view
+		bEnableUpdateRateOptimizations = SkeletalMeshComponent->bEnableUpdateRateOptimizations;
+		MeshComponentUpdateFlag = SkeletalMeshComponent->MeshComponentUpdateFlag;
+
+		SkeletalMeshComponent->bEnableUpdateRateOptimizations = false;
+		SkeletalMeshComponent->MeshComponentUpdateFlag = EMeshComponentUpdateFlag::AlwaysTickPoseAndRefreshBones;
+
+		ComponentTransform = SkeletalMeshComponent->GetComponentToWorld().GetRelativeTransform(SkeletalMeshComponent->GetOwner()->GetTransform());
+
+		if(!Sequence.IsValid())
+		{
+			// build an asset path
+			const USequenceRecorderSettings* Settings = GetDefault<USequenceRecorderSettings>();
+
+			FString AssetPath(TEXT("/Game"));
+			AssetPath /= Settings->SequenceRecordingBasePath.Path;
+			if(Settings->AnimationSubDirectory.Len() > 0)
+			{
+				AssetPath /= Settings->AnimationSubDirectory;
+			}
+
+			FString AssetName = Settings->SequenceName.Len() > 0 ? Settings->SequenceName : TEXT("RecordedSequence");
+			AssetName += TEXT("_");
+			AssetName += Actor->GetActorLabel();
+
+			Sequence = SequenceRecorderUtils::MakeNewAsset<UAnimSequence>(AssetPath, AssetName);
+			if(Sequence.IsValid())
+			{
+				FAssetRegistryModule::AssetCreated(Sequence.Get());
+
+				// set skeleton
+				Sequence->SetSkeleton(SkeletalMeshComponent->SkeletalMesh->Skeleton);
+			}
+		}
+
+		if(Sequence.IsValid())
+		{
+			FAnimationRecorderManager::Get().RecordAnimation(SkeletalMeshComponent.Get(), Sequence.Get(), AnimationSettings);
+
+			if(MovieScene)
+			{
+				UMovieSceneSkeletalAnimationTrack* AnimTrack = MovieScene->AddTrack<UMovieSceneSkeletalAnimationTrack>(Guid);
+				if(AnimTrack)
+				{
+					AnimTrack->AddNewAnimation(Time, Sequence.Get());
+					MovieSceneSection = Cast<UMovieSceneSkeletalAnimationSection>(AnimTrack->GetAllSections()[0]);
+				}
+			}
+		}
+	}
+
+	bRecording = bRecord;
+}
+
+void FMovieSceneAnimationPropertyRecorder::FinalizeSection()
+{
+	bRecording = false;
+
+	if(Sequence.IsValid())
+	{
+		// enable root motion on the animation
+		Sequence->bEnableRootMotion = true;
+		Sequence->RootMotionRootLock = ERootMotionRootLock::Zero;
+	}
+
+	if(SkeletalMeshComponent.IsValid())
+	{
+		// restore update flags
+		SkeletalMeshComponent->bEnableUpdateRateOptimizations = bEnableUpdateRateOptimizations;
+		SkeletalMeshComponent->MeshComponentUpdateFlag = MeshComponentUpdateFlag;
+
+		// only show a message if we dont have a valid movie section
+		const bool bShowMessage = !MovieSceneSection.IsValid();
+		FAnimationRecorderManager::Get().StopRecordingAnimation(SkeletalMeshComponent.Get(), bShowMessage);
+	}
+
+	if(MovieSceneSection.IsValid() && Sequence.IsValid())
+	{
+		MovieSceneSection->SetEndTime(MovieSceneSection->GetStartTime() + Sequence->GetPlayLength());
+	}
+}
+
+void FMovieSceneAnimationPropertyRecorder::Record(float CurrentTime)
+{
+	// The animation recorder does most of the work here
+
+	if(SkeletalMeshComponent.IsValid())
+	{
+		// re-force updates on as gameplay can sometimes turn these back off!
+		SkeletalMeshComponent->bEnableUpdateRateOptimizations = false;
+		SkeletalMeshComponent->MeshComponentUpdateFlag = EMeshComponentUpdateFlag::AlwaysTickPoseAndRefreshBones;
+	}
+}
\ No newline at end of file
diff --git a/Engine/Source/Editor/SequenceRecorder/Private/PropertyRecorders/MovieSceneAnimationPropertyRecorder.h b/Engine/Source/Editor/SequenceRecorder/Private/PropertyRecorders/MovieSceneAnimationPropertyRecorder.h
new file mode 100644
--- /dev/null
+++ b/Engine/Source/Editor/SequenceRecorder/Private/PropertyRecorders/MovieSceneAnimationPropertyRecorder.h
@@ -0,0 +1,53 @@
+// Copyright 1998-2015 Epic Games, Inc. All Rights Reserved.
+
+#pragma once
+
+#include "IMovieScenePropertyRecorder.h"
+
+class FMovieSceneAnimationPropertyRecorder : public IMovieScenePropertyRecorder
+{
+public:
+	FMovieSceneAnimationPropertyRecorder(FAnimationRecordingSettings& InAnimationSettings, UAnimSequence* InSpecifiedSequence);
+	virtual ~FMovieSceneAnimationPropertyRecorder() {}
+
+	virtual void CreateSection(UObject* InObjectToRecord, class UMovieScene* MovieScene, const FGuid& Guid, float Time, bool bRecord) override;
+	virtual void FinalizeSection() override;
+	virtual void Record(float CurrentTime) override;
+	virtual void InvalidateObjectToRecord() override
+	{
+		ObjectToRecord = nullptr;
+	}
+
+	UAnimSequence* GetSequence() const { return Sequence.Get(); }
+
+	USkeletalMesh* GetSkeletalMesh() const { return SkeletalMesh.Get(); }
+
+	const FTransform& GetComponentTransform() const { return ComponentTransform; }
+
+private:
+	/** Object to record from */
+	TLazyObjectPtr<UObject> ObjectToRecord;
+
+	/** Section to record to */
+	TWeakObjectPtr<class UMovieSceneSkeletalAnimationSection> MovieSceneSection;
+
+	TWeakObjectPtr<class UAnimSequence> Sequence;
+
+	TWeakObjectPtr<class USkeletalMeshComponent> SkeletalMeshComponent;
+
+	TWeakObjectPtr<class USkeletalMesh> SkeletalMesh;
+
+	/** Local transform of the component we are recording */
+	FTransform ComponentTransform;
+
+	FAnimationRecordingSettings AnimationSettings;
+
+	/** Used to store/restore update flag when recording */
+	EMeshComponentUpdateFlag::Type MeshComponentUpdateFlag;
+
+	/** Used to store/restore URO when recording */
+	bool bEnableUpdateRateOptimizations;
+
+	/** Flag if we are actually recording or not */
+	bool bRecording;
+};
\ No newline at end of file
diff --git a/Engine/Source/Editor/SequenceRecorder/Private/PropertyRecorders/MovieSceneParticleTrackPropertyRecorder.cpp b/Engine/Source/Editor/SequenceRecorder/Private/PropertyRecorders/MovieSceneParticleTrackPropertyRecorder.cpp
new file mode 100644
--- /dev/null
+++ b/Engine/Source/Editor/SequenceRecorder/Private/PropertyRecorders/MovieSceneParticleTrackPropertyRecorder.cpp
@@ -0,0 +1,99 @@
+// Copyright 1998-2015 Epic Games, Inc. All Rights Reserved.
+
+#include "SequenceRecorderPrivatePCH.h"
+#include "MovieSceneParticleTrackPropertyRecorder.h"
+#include "MovieSceneParticleSection.h"
+#include "MovieSceneParticleTrack.h"
+#include "MovieScene.h"
+#include "Particles/ParticleSystemComponent.h"
+
+FMovieSceneParticleTrackPropertyRecorder::~FMovieSceneParticleTrackPropertyRecorder()
+{
+	if(DelegateProxy.IsValid())
+	{
+		DelegateProxy->PropertyRecorder = nullptr;
+		DelegateProxy->RemoveFromRoot();
+		DelegateProxy.Reset();
+	}
+}
+
+void FMovieSceneParticleTrackPropertyRecorder::CreateSection(UObject* InObjectToRecord, UMovieScene* MovieScene, const FGuid& Guid, float Time, bool bRecord)
+{
+	SystemToRecord = CastChecked<UParticleSystemComponent>(InObjectToRecord);
+
+	UMovieSceneParticleTrack* ParticleTrack = MovieScene->AddTrack<UMovieSceneParticleTrack>(Guid);
+	if(ParticleTrack)
+	{
+		MovieSceneSection = Cast<UMovieSceneParticleSection>(ParticleTrack->CreateNewSection());
+
+		ParticleTrack->AddSection(*MovieSceneSection);
+
+		MovieSceneSection->SetStartTime(Time);
+
+		bWasTriggered = false;
+
+		DelegateProxy = NewObject<UMovieSceneParticleTrackPropertyRecorder>();
+		DelegateProxy->PropertyRecorder = this;
+		DelegateProxy->AddToRoot();
+		UParticleSystemComponent::OnSystemPreActivationChange.AddUObject(DelegateProxy.Get(), &UMovieSceneParticleTrackPropertyRecorder::OnTriggered);
+	}
+
+	bRecording = bRecord;
+
+	PreviousState = EParticleKey::Deactivate;
+
+	Record(Time);
+}
+
+void FMovieSceneParticleTrackPropertyRecorder::FinalizeSection()
+{
+	bRecording = false;
+}
+
+void FMovieSceneParticleTrackPropertyRecorder::Record(float CurrentTime)
+{
+	if(SystemToRecord.IsValid())
+	{
+		MovieSceneSection->SetEndTime(CurrentTime);
+
+		if(bRecording)
+		{
+			EParticleKey::Type NewState = EParticleKey::Deactivate;
+			if(SystemToRecord->IsRegistered() && SystemToRecord->IsActive() && !SystemToRecord->bWasDeactivated)
+			{
+				if(bWasTriggered)
+				{
+					NewState = EParticleKey::Trigger;
+					bWasTriggered = false;
+				}
+				else
+				{
+					NewState = EParticleKey::Activate;
+				}
+			}
+			else
+			{
+				NewState = EParticleKey::Deactivate;
+			}
+
+			if(NewState != PreviousState)
+			{
+				MovieSceneSection->AddKey(CurrentTime, NewState);
+			}
+
+			if(NewState == EParticleKey::Trigger)
+			{
+				NewState = EParticleKey::Activate;
+			}
+			PreviousState = NewState;
+		}
+	}
+}
+
+void UMovieSceneParticleTrackPropertyRecorder::OnTriggered(UParticleSystemComponent* Component, bool bActivating)
+{ 
+	if(PropertyRecorder && PropertyRecorder->SystemToRecord.Get() == Component)
+	{
+		PropertyRecorder->bWasTriggered = bActivating;
+	}
+}
\ No newline at end of file
diff --git a/Engine/Source/Editor/SequenceRecorder/Private/PropertyRecorders/MovieSceneParticleTrackPropertyRecorder.h b/Engine/Source/Editor/SequenceRecorder/Private/PropertyRecorders/MovieSceneParticleTrackPropertyRecorder.h
new file mode 100644
--- /dev/null
+++ b/Engine/Source/Editor/SequenceRecorder/Private/PropertyRecorders/MovieSceneParticleTrackPropertyRecorder.h
@@ -0,0 +1,58 @@
+// Copyright 1998-2015 Epic Games, Inc. All Rights Reserved.
+
+#pragma once
+
+#include "IMovieScenePropertyRecorder.h"
+#include "MovieSceneParticleSection.h"
+#include "MovieSceneParticleTrackPropertyRecorder.generated.h"
+
+struct FMovieSceneParticleTrackPropertyRecorder;
+
+UCLASS(MinimalAPI, NotBlueprintable)
+class UMovieSceneParticleTrackPropertyRecorder : public UObject
+{
+	GENERATED_BODY()
+
+public:
+	UMovieSceneParticleTrackPropertyRecorder()
+		: PropertyRecorder(nullptr)
+	{
+	}
+
+	UFUNCTION()
+	void OnTriggered(UParticleSystemComponent* Component, bool bActivating);
+
+	FMovieSceneParticleTrackPropertyRecorder* PropertyRecorder;
+};
+
+struct FMovieSceneParticleTrackPropertyRecorder : public IMovieScenePropertyRecorder
+{
+public:
+	friend class UMovieSceneParticleTrackPropertyRecorder;
+
+	virtual ~FMovieSceneParticleTrackPropertyRecorder();
+
+	virtual void CreateSection(UObject* InObjectToRecord, class UMovieScene* MovieScene, const FGuid& Guid, float Time, bool bRecord) override;
+	virtual void FinalizeSection() override;
+	virtual void Record(float CurrentTime) override;
+	virtual void InvalidateObjectToRecord() override
+	{
+		SystemToRecord = nullptr;
+	}
+
+private:
+	/** Object to record from */
+	TLazyObjectPtr<class UParticleSystemComponent> SystemToRecord;
+
+	/** Section to record to */
+	TWeakObjectPtr<class UMovieSceneParticleSection> MovieSceneSection;
+
+	/** Flag if we are actually recording or not */
+	bool bRecording;
+
+	bool bWasTriggered;
+
+	EParticleKey::Type PreviousState;
+
+	TWeakObjectPtr<UMovieSceneParticleTrackPropertyRecorder> DelegateProxy;
+};
diff --git a/Engine/Source/Editor/SequenceRecorder/Private/PropertyRecorders/MovieSceneSpawnPropertyRecorder.cpp b/Engine/Source/Editor/SequenceRecorder/Private/PropertyRecorders/MovieSceneSpawnPropertyRecorder.cpp
new file mode 100644
--- /dev/null
+++ b/Engine/Source/Editor/SequenceRecorder/Private/PropertyRecorders/MovieSceneSpawnPropertyRecorder.cpp
@@ -0,0 +1,62 @@
+// Copyright 1998-2015 Epic Games, Inc. All Rights Reserved.
+
+#include "SequenceRecorderPrivatePCH.h"
+#include "MovieSceneSpawnPropertyRecorder.h"
+#include "MovieSceneBoolSection.h"
+#include "MovieSceneSpawnTrack.h"
+#include "MovieScene.h"
+
+void FMovieSceneSpawnPropertyRecorder::CreateSection(UObject* InObjectToRecord, UMovieScene* MovieScene, const FGuid& Guid, float Time, bool bRecord)
+{
+	ObjectToRecord = InObjectToRecord;
+
+	UMovieSceneSpawnTrack* SpawnTrack = MovieScene->AddTrack<UMovieSceneSpawnTrack>(Guid);
+	if(SpawnTrack)
+	{
+		MovieSceneSection = Cast<UMovieSceneBoolSection>(SpawnTrack->CreateNewSection());
+
+		SpawnTrack->AddSection(*MovieSceneSection);
+		SpawnTrack->SetObjectId(Guid);
+
+		MovieSceneSection->SetDefault(false);
+		MovieSceneSection->AddKey(0.0f, false, EMovieSceneKeyInterpolation::Break);
+
+		MovieSceneSection->SetStartTime(Time);
+		MovieSceneSection->SetIsInfinite(true);
+	}
+
+	bRecording = bRecord;
+
+	bWasSpawned = false;
+
+	Record(Time);
+}
+
+void FMovieSceneSpawnPropertyRecorder::FinalizeSection()
+{
+	bRecording = false;
+
+	const bool bSpawned = ObjectToRecord.IsValid();
+	if(bSpawned != bWasSpawned)
+	{
+		MovieSceneSection->AddKey(MovieSceneSection->GetEndTime(), bSpawned, EMovieSceneKeyInterpolation::Break);
+	}
+}
+
+void FMovieSceneSpawnPropertyRecorder::Record(float CurrentTime)
+{
+	if(bRecording)
+	{
+		if(ObjectToRecord.IsValid())
+		{
+			MovieSceneSection->SetEndTime(CurrentTime);
+		}
+
+		const bool bSpawned = ObjectToRecord.IsValid();
+		if(bSpawned != bWasSpawned)
+		{
+			MovieSceneSection->AddKey(CurrentTime, bSpawned, EMovieSceneKeyInterpolation::Break);
+		}
+		bWasSpawned = bSpawned;
+	}
+}
\ No newline at end of file
diff --git a/Engine/Source/Editor/SequenceRecorder/Private/PropertyRecorders/MovieSceneSpawnPropertyRecorder.h b/Engine/Source/Editor/SequenceRecorder/Private/PropertyRecorders/MovieSceneSpawnPropertyRecorder.h
new file mode 100644
--- /dev/null
+++ b/Engine/Source/Editor/SequenceRecorder/Private/PropertyRecorders/MovieSceneSpawnPropertyRecorder.h
@@ -0,0 +1,31 @@
+// Copyright 1998-2015 Epic Games, Inc. All Rights Reserved.
+
+#pragma once
+
+#include "IMovieScenePropertyRecorder.h"
+
+class FMovieSceneSpawnPropertyRecorder : public IMovieScenePropertyRecorder
+{
+public:
+	virtual ~FMovieSceneSpawnPropertyRecorder() {}
+
+	virtual void CreateSection(UObject* InObjectToRecord, class UMovieScene* MovieScene, const FGuid& Guid, float Time, bool bRecord) override;
+	virtual void FinalizeSection() override;
+	virtual void Record(float CurrentTime) override;
+	virtual void InvalidateObjectToRecord() override
+	{
+		ObjectToRecord = nullptr;
+	}
+
+private:
+	/** Object to record from */
+	TLazyObjectPtr<UObject> ObjectToRecord;
+
+	/** Section to record to */
+	TWeakObjectPtr<class UMovieSceneBoolSection> MovieSceneSection;
+
+	/** Flag if we are actually recording or not */
+	bool bRecording;
+
+	bool bWasSpawned;
+};
\ No newline at end of file
diff --git a/Engine/Source/Editor/SequenceRecorder/Private/PropertyRecorders/MovieSceneVisibilityPropertyRecorder.cpp b/Engine/Source/Editor/SequenceRecorder/Private/PropertyRecorders/MovieSceneVisibilityPropertyRecorder.cpp
new file mode 100644
--- /dev/null
+++ b/Engine/Source/Editor/SequenceRecorder/Private/PropertyRecorders/MovieSceneVisibilityPropertyRecorder.cpp
@@ -0,0 +1,76 @@
+// Copyright 1998-2015 Epic Games, Inc. All Rights Reserved.
+
+#include "SequenceRecorderPrivatePCH.h"
+#include "MovieSceneVisibilityPropertyRecorder.h"
+#include "MovieSceneVisibilitySection.h"
+#include "MovieSceneVisibilityTrack.h"
+#include "MovieScene.h"
+
+static const FName VisibilityTrackName = TEXT("Visibility");
+
+void FMovieSceneVisibilityPropertyRecorder::CreateSection(UObject* InObjectToRecord, UMovieScene* MovieScene, const FGuid& Guid, float Time, bool bRecord)
+{
+	ObjectToRecord = InObjectToRecord;
+
+	UMovieSceneVisibilityTrack* VisibilityTrack = MovieScene->AddTrack<UMovieSceneVisibilityTrack>(Guid);
+	if(VisibilityTrack)
+	{
+		VisibilityTrack->SetPropertyNameAndPath(VisibilityTrackName, VisibilityTrackName.ToString());
+
+		MovieSceneSection = Cast<UMovieSceneVisibilitySection>(VisibilityTrack->CreateNewSection());
+
+		VisibilityTrack->AddSection(*MovieSceneSection);
+
+		MovieSceneSection->SetDefault(false);
+
+		bWasVisible = false;
+		if(USceneComponent* SceneComponent = Cast<USceneComponent>(ObjectToRecord.Get()))
+		{
+			bWasVisible = SceneComponent->IsVisible() && SceneComponent->IsRegistered();
+		}
+		else if(AActor* Actor = Cast<AActor>(ObjectToRecord.Get()))
+		{
+			bWasVisible = !Actor->bHidden;
+		}
+
+		MovieSceneSection->AddKey(Time, bWasVisible, EMovieSceneKeyInterpolation::Break);
+
+		MovieSceneSection->SetStartTime(Time);
+	}
+
+	bRecording = bRecord;
+
+	Record(Time);
+}
+
+void FMovieSceneVisibilityPropertyRecorder::FinalizeSection()
+{
+	bRecording = false;
+}
+
+void FMovieSceneVisibilityPropertyRecorder::Record(float CurrentTime)
+{
+	if(ObjectToRecord.IsValid())
+	{
+		MovieSceneSection->SetEndTime(CurrentTime);
+
+		if(bRecording)
+		{
+			bool bVisible = false;
+			if(USceneComponent* SceneComponent = Cast<USceneComponent>(ObjectToRecord.Get()))
+			{
+				bVisible = SceneComponent->IsVisible() && SceneComponent->IsRegistered();
+			}
+			else if(AActor* Actor = Cast<AActor>(ObjectToRecord.Get()))
+			{
+				bVisible = !Actor->bHidden;
+			}
+
+			if(bVisible != bWasVisible)
+			{
+				MovieSceneSection->AddKey(CurrentTime, bVisible, EMovieSceneKeyInterpolation::Break);
+			}
+			bWasVisible = bVisible;
+		}
+	}
+}
\ No newline at end of file
diff --git a/Engine/Source/Editor/SequenceRecorder/Private/PropertyRecorders/MovieSceneVisibilityPropertyRecorder.h b/Engine/Source/Editor/SequenceRecorder/Private/PropertyRecorders/MovieSceneVisibilityPropertyRecorder.h
new file mode 100644
--- /dev/null
+++ b/Engine/Source/Editor/SequenceRecorder/Private/PropertyRecorders/MovieSceneVisibilityPropertyRecorder.h
@@ -0,0 +1,32 @@
+// Copyright 1998-2015 Epic Games, Inc. All Rights Reserved.
+
+#pragma once
+
+#include "IMovieScenePropertyRecorder.h"
+
+class FMovieSceneVisibilityPropertyRecorder : public IMovieScenePropertyRecorder
+{
+public:
+	virtual ~FMovieSceneVisibilityPropertyRecorder() {}
+
+	virtual void CreateSection(UObject* InObjectToRecord, class UMovieScene* MovieScene, const FGuid& Guid, float Time, bool bRecord) override;
+	virtual void FinalizeSection() override;
+	virtual void Record(float CurrentTime) override;
+	virtual void InvalidateObjectToRecord() override
+	{
+		ObjectToRecord = nullptr;
+	}
+
+private:
+	/** Object to record from */
+	TLazyObjectPtr<class UObject> ObjectToRecord;
+
+	/** Section to record to */
+	TWeakObjectPtr<class UMovieSceneVisibilitySection> MovieSceneSection;
+
+	/** Flag if we are actually recording or not */
+	bool bRecording;
+
+	/** Flag used to track visibility state and add keys when this changes */
+	bool bWasVisible;
+};
\ No newline at end of file
diff --git a/Engine/Source/Editor/SequenceRecorder/Private/SSequenceRecorder.cpp b/Engine/Source/Editor/SequenceRecorder/Private/SSequenceRecorder.cpp
new file mode 100644
--- /dev/null
+++ b/Engine/Source/Editor/SequenceRecorder/Private/SSequenceRecorder.cpp
@@ -0,0 +1,361 @@
+// Copyright 1998-2016 Epic Games, Inc. All Rights Reserved.
+
+#include "SequenceRecorderPrivatePCH.h"
+#include "SSequenceRecorder.h"
+#include "ActorRecording.h"
+#include "SequenceRecorderCommands.h"
+#include "SNumericEntryBox.h"
+#include "SequenceRecorderSettings.h"
+#include "SequenceRecorder.h"
+
+#define LOCTEXT_NAMESPACE "SequenceRecorder"
+
+static const FName ActorColumnName(TEXT("Actor"));
+static const FName AnimationColumnName(TEXT("Animation"));
+static const FName LengthColumnName(TEXT("Length"));
+
+/** A widget to display information about an animation recording in the list view */
+class SSequenceRecorderListRow : public SMultiColumnTableRow<UActorRecording*>
+{
+public:
+	SLATE_BEGIN_ARGS(SSequenceRecorderListRow) {}
+
+	/** The list item for this row */
+	SLATE_ARGUMENT(UActorRecording*, Recording)
+
+	SLATE_END_ARGS()
+
+	void Construct(const FArguments& Args, const TSharedRef<STableViewBase>& OwnerTableView)
+	{
+		RecordingPtr = Args._Recording;
+
+		SMultiColumnTableRow<UActorRecording*>::Construct(
+			FSuperRowType::FArguments()
+			.Padding(1.0f),
+			OwnerTableView
+			);
+	}
+
+	/** Overridden from SMultiColumnTableRow.  Generates a widget for this column of the list view. */
+	virtual TSharedRef<SWidget> GenerateWidgetForColumn(const FName& ColumnName) override
+	{
+		if (ColumnName == ActorColumnName)
+		{
+			return
+				SNew(STextBlock)
+				.Text(TAttribute<FText>::Create(TAttribute<FText>::FGetter::CreateSP(this, &SSequenceRecorderListRow::GetRecordingActorName)));
+		}
+		else if (ColumnName == AnimationColumnName)
+		{
+			return
+				SNew(STextBlock)
+				.Text(TAttribute<FText>::Create(TAttribute<FText>::FGetter::CreateSP(this, &SSequenceRecorderListRow::GetRecordingAnimationName)));
+		}
+		else if (ColumnName == LengthColumnName)
+		{
+			return
+				SNew(STextBlock)
+				.Text(TAttribute<FText>::Create(TAttribute<FText>::FGetter::CreateSP(this, &SSequenceRecorderListRow::GetRecordingLengthText)));
+		}
+
+		return SNullWidget::NullWidget;
+	}
+
+private:
+	FText GetRecordingActorName() const
+	{
+		FText ActorName(LOCTEXT("InvalidActorName", "None"));
+		if (RecordingPtr.IsValid() && RecordingPtr.Get()->ActorToRecord.IsValid())
+		{
+			ActorName = FText::FromString(RecordingPtr.Get()->ActorToRecord.Get()->GetActorLabel());
+		}
+
+		return ActorName;
+	}
+
+	FText GetRecordingAnimationName() const
+	{
+		FText AnimationName(LOCTEXT("InvalidAnimationName", "None"));
+		if (RecordingPtr.IsValid())
+		{
+			if(!RecordingPtr.Get()->bSpecifyTargetAnimation)
+			{
+				AnimationName = LOCTEXT("AutoCreatedAnimationName", "Auto-created");
+			}
+			else if(RecordingPtr.Get()->TargetAnimation.IsValid())
+			{
+				AnimationName = FText::FromString(RecordingPtr.Get()->TargetAnimation.Get()->GetName());
+			}
+		}
+
+		return AnimationName;
+	}
+
+	FText GetRecordingLengthText() const
+	{
+		FText LengthText(LOCTEXT("InvalidLengthName", "None"));
+		if (RecordingPtr.IsValid())
+		{
+			LengthText = FText::AsNumber(RecordingPtr.Get()->AnimationSettings.Length);
+		}
+
+		return LengthText;
+	}
+
+	TWeakObjectPtr<UActorRecording> RecordingPtr;
+};
+
+
+void SSequenceRecorder::Construct(const FArguments& Args)
+{
+	CommandList = MakeShareable(new FUICommandList);
+
+	BindCommands();
+
+	FPropertyEditorModule& PropertyEditorModule = FModuleManager::Get().GetModuleChecked<FPropertyEditorModule>("PropertyEditor");
+
+	FDetailsViewArgs DetailsViewArgs;
+	DetailsViewArgs.NameAreaSettings = FDetailsViewArgs::HideNameArea;
+	DetailsViewArgs.bAllowSearch = false;
+
+	ActorRecordingDetailsView = PropertyEditorModule.CreateDetailView( DetailsViewArgs );
+	SequenceRecordingDetailsView = PropertyEditorModule.CreateDetailView( DetailsViewArgs );
+	SequenceRecordingDetailsView->SetObject(GetMutableDefault<USequenceRecorderSettings>());
+
+	FToolBarBuilder ToolBarBuilder(CommandList, FMultiBoxCustomization::None);
+
+	ToolBarBuilder.BeginSection(TEXT("Recording"));
+	{
+		ToolBarBuilder.AddToolBarButton(FSequenceRecorderCommands::Get().RecordAll);
+		ToolBarBuilder.AddToolBarButton(FSequenceRecorderCommands::Get().StopAll);
+	}
+	ToolBarBuilder.EndSection();
+
+	ToolBarBuilder.BeginSection(TEXT("RecordingManagement"));
+	{
+		ToolBarBuilder.AddToolBarButton(FSequenceRecorderCommands::Get().AddRecording);
+		ToolBarBuilder.AddToolBarButton(FSequenceRecorderCommands::Get().RemoveRecording);
+	}
+	ToolBarBuilder.EndSection();
+
+	ChildSlot
+	[
+		SNew(SSplitter)
+		.Orientation(Orient_Vertical)
+		+SSplitter::Slot()
+		.Value(0.33f)
+		[
+			SNew(SVerticalBox)
+			+SVerticalBox::Slot()
+			.AutoHeight()
+			.Padding(FMargin(0.0f, 4.0f, 0.0f, 0.0f))
+			[
+				ToolBarBuilder.MakeWidget()
+			]
+			+SVerticalBox::Slot()
+			.FillHeight(1.0f)
+			.Padding(FMargin(0.0f, 4.0f, 0.0f, 0.0f))
+			[
+				SNew(SBorder)
+				.BorderImage(FEditorStyle::GetBrush("ToolPanel.GroupBorder"))
+				.Padding(FMargin(4.0f, 4.0f))
+				[
+					SNew(SOverlay)
+					+SOverlay::Slot()
+					[
+						SAssignNew(ListView, SListView<UActorRecording*>)
+						.ListItemsSource(&FSequenceRecorder::Get().GetQueuedRecordings())
+						.SelectionMode(ESelectionMode::SingleToggle)
+						.OnGenerateRow(this, &SSequenceRecorder::MakeListViewWidget)
+						.OnSelectionChanged(this, &SSequenceRecorder::OnSelectionChanged)
+						.HeaderRow
+						(
+						SNew(SHeaderRow)
+						+ SHeaderRow::Column(ActorColumnName)
+						.FillWidth(43.0f)
+						.DefaultLabel(LOCTEXT("ActorHeaderName", "Actor"))
+						+ SHeaderRow::Column(AnimationColumnName)
+						.FillWidth(43.0f)
+						.DefaultLabel(LOCTEXT("AnimationHeaderName", "Animation"))
+						+ SHeaderRow::Column(LengthColumnName)
+						.FillWidth(14.0f)
+						.DefaultLabel(LOCTEXT("LengthHeaderName", "Length"))
+						)
+					]
+					+SOverlay::Slot()
+					[
+						SNew(SVerticalBox)
+						+SVerticalBox::Slot()
+						.VAlign(VAlign_Bottom)
+						.MaxHeight(2.0f)
+						[
+							SAssignNew(DelayProgressBar, SProgressBar)
+							.Percent(this, &SSequenceRecorder::GetDelayPercent)
+							.Visibility(this, &SSequenceRecorder::GetDelayProgressVisibilty)
+						]
+					]
+				]
+			]
+		]
+		+SSplitter::Slot()
+		.Value(0.66f)
+		[
+			SNew(SScrollBox)
+			+SScrollBox::Slot()
+			[
+				SNew(SVerticalBox)
+				.IsEnabled_Lambda([]() { return !FSequenceRecorder::Get().IsRecording(); })
+				+SVerticalBox::Slot()
+				.AutoHeight()
+				[
+					SequenceRecordingDetailsView.ToSharedRef()
+				]
+				+SVerticalBox::Slot()
+				.AutoHeight()
+				[
+					ActorRecordingDetailsView.ToSharedRef()
+				]	
+			]
+		]
+	];
+
+	// Register refresh timer
+	if (!ActiveTimerHandle.IsValid())
+	{
+		ActiveTimerHandle = RegisterActiveTimer(0.1f, FWidgetActiveTimerDelegate::CreateSP(this, &SSequenceRecorder::HandleRefreshItems));
+	}
+}
+
+void SSequenceRecorder::BindCommands()
+{
+	CommandList->MapAction(FSequenceRecorderCommands::Get().RecordAll,
+		FExecuteAction::CreateSP(this, &SSequenceRecorder::HandleRecord),
+		FCanExecuteAction::CreateSP(this, &SSequenceRecorder::CanRecord),
+		FIsActionChecked(),
+		FIsActionButtonVisible::CreateSP(this, &SSequenceRecorder::IsRecordVisible)
+		);
+
+	CommandList->MapAction(FSequenceRecorderCommands::Get().StopAll,
+		FExecuteAction::CreateSP(this, &SSequenceRecorder::HandleStopAll),
+		FCanExecuteAction::CreateSP(this, &SSequenceRecorder::CanStopAll),
+		FIsActionChecked(),
+		FIsActionButtonVisible::CreateSP(this, &SSequenceRecorder::IsStopAllVisible)
+		);
+
+	CommandList->MapAction(FSequenceRecorderCommands::Get().AddRecording,
+		FExecuteAction::CreateSP(this, &SSequenceRecorder::HandleAddRecording),
+		FCanExecuteAction::CreateSP(this, &SSequenceRecorder::CanAddRecording)
+		);
+
+	CommandList->MapAction(FSequenceRecorderCommands::Get().RemoveRecording,
+		FExecuteAction::CreateSP(this, &SSequenceRecorder::HandleRemoveRecording),
+		FCanExecuteAction::CreateSP(this, &SSequenceRecorder::CanRemoveRecording)
+		);
+}
+
+TSharedRef<ITableRow> SSequenceRecorder::MakeListViewWidget(UActorRecording* Recording, const TSharedRef<STableViewBase>& OwnerTable) const
+{
+	return
+		SNew(SSequenceRecorderListRow, OwnerTable)
+		.Recording(Recording);
+}
+
+void SSequenceRecorder::OnSelectionChanged(UActorRecording* Recording, ESelectInfo::Type SelectionType) const
+{
+	if(Recording)
+	{
+		ActorRecordingDetailsView->SetObject(Recording);
+	}
+	else
+	{
+		ActorRecordingDetailsView->SetObject(nullptr);
+	}
+}
+
+void SSequenceRecorder::HandleRecord()
+{
+	FSequenceRecorder::Get().StartRecording();
+}
+
+bool SSequenceRecorder::CanRecord() const
+{
+	return FSequenceRecorder::Get().HasQueuedRecordings();
+}
+
+bool SSequenceRecorder::IsRecordVisible() const
+{
+	return !FAnimationRecorderManager::Get().IsRecording() && !FSequenceRecorder::Get().IsDelaying();
+}
+
+void SSequenceRecorder::HandleStopAll()
+{
+	FSequenceRecorder::Get().StopRecording();
+}
+
+bool SSequenceRecorder::CanStopAll() const
+{
+	return FAnimationRecorderManager::Get().IsRecording() || FSequenceRecorder::Get().IsDelaying();
+}
+
+bool SSequenceRecorder::IsStopAllVisible() const
+{
+	return FAnimationRecorderManager::Get().IsRecording() || FSequenceRecorder::Get().IsDelaying();
+}
+
+void SSequenceRecorder::HandleAddRecording()
+{
+	FSequenceRecorder::Get().AddNewQueuedRecording();
+}
+
+bool SSequenceRecorder::CanAddRecording() const
+{
+	return !FAnimationRecorderManager::Get().IsRecording();
+}
+
+void SSequenceRecorder::HandleRemoveRecording()
+{
+	TArray<UActorRecording*> SelectedRecordings;
+	ListView->GetSelectedItems(SelectedRecordings);
+	UActorRecording* SelectedRecording = SelectedRecordings.Num() > 0 ? SelectedRecordings[0] : nullptr;
+	if(SelectedRecording)
+	{
+		FSequenceRecorder::Get().RemoveQueuedRecording(SelectedRecording);
+
+		TArray<TWeakObjectPtr<UObject>> SelectedObjects = ActorRecordingDetailsView->GetSelectedObjects();
+		if(SelectedObjects.Num() > 0 && SelectedObjects[0].Get() == SelectedRecording)
+		{
+			ActorRecordingDetailsView->SetObject(nullptr);
+		}
+	}
+}
+
+bool SSequenceRecorder::CanRemoveRecording() const
+{
+	return ListView->GetNumItemsSelected() > 0 && !FAnimationRecorderManager::Get().IsRecording();
+}
+
+EActiveTimerReturnType SSequenceRecorder::HandleRefreshItems(double InCurrentTime, float InDeltaTime)
+{
+	if(FSequenceRecorder::Get().AreQueuedRecordingsDirty())
+	{
+		ListView->RequestListRefresh();
+
+		FSequenceRecorder::Get().ResetQueuedRecordingsDirty();
+	}
+
+	return EActiveTimerReturnType::Continue;
+}
+
+TOptional<float> SSequenceRecorder::GetDelayPercent() const
+{
+	const float Delay = GetDefault<USequenceRecorderSettings>()->RecordingDelay;
+	const float Countdown = FSequenceRecorder::Get().GetCurrentDelay();
+	return Delay > 0.0f ? Countdown / Delay : 0.0f;
+}
+
+EVisibility SSequenceRecorder::GetDelayProgressVisibilty() const
+{
+	return FSequenceRecorder::Get().IsDelaying() ? EVisibility::Visible : EVisibility::Hidden;
+}
+
+#undef LOCTEXT_NAMESPACE
\ No newline at end of file
diff --git a/Engine/Source/Editor/SequenceRecorder/Private/SSequenceRecorder.h b/Engine/Source/Editor/SequenceRecorder/Private/SSequenceRecorder.h
new file mode 100644
--- /dev/null
+++ b/Engine/Source/Editor/SequenceRecorder/Private/SSequenceRecorder.h
@@ -0,0 +1,70 @@
+// Copyright 1998-2016 Epic Games, Inc. All Rights Reserved.
+
+#pragma once
+
+#include "ActorRecording.h"
+
+class SSequenceRecorder : public SCompoundWidget
+{
+	SLATE_BEGIN_ARGS(SSequenceRecorder)
+	{
+	}
+
+	SLATE_END_ARGS()
+
+	void Construct(const FArguments& Args);
+
+private:
+
+	void BindCommands();
+
+	TSharedRef<ITableRow> MakeListViewWidget(class UActorRecording* Recording, const TSharedRef<STableViewBase>& OwnerTable) const;
+
+	FText GetRecordingActorName(class UActorRecording* Recording) const;
+
+	void OnSelectionChanged(UActorRecording* Recording, ESelectInfo::Type SelectionType) const;
+
+	void HandleRecord();
+
+	EActiveTimerReturnType StartDelayedRecord(double InCurrentTime, float InDeltaTime);
+
+	bool CanRecord() const;
+
+	bool IsRecordVisible() const;
+
+	void HandleStopAll();
+
+	bool CanStopAll() const;
+
+	bool IsStopAllVisible() const;
+
+	void HandleAddRecording();
+
+	bool CanAddRecording() const;
+
+	void HandleRemoveRecording();
+
+	bool CanRemoveRecording() const;
+
+	EActiveTimerReturnType HandleRefreshItems(double InCurrentTime, float InDeltaTime);
+
+	TOptional<float> GetDelayPercent() const;
+
+	void OnDelayChanged(float NewValue);
+
+	EVisibility GetDelayProgressVisibilty() const;
+
+private:
+	TSharedPtr<IDetailsView> SequenceRecordingDetailsView;
+
+	TSharedPtr<IDetailsView> ActorRecordingDetailsView;
+
+	TSharedPtr<SListView<UActorRecording*>> ListView;
+
+	TSharedPtr<FUICommandList> CommandList;
+
+	/** The handle to the refresh tick timer */
+	TWeakPtr<FActiveTimerHandle> ActiveTimerHandle;
+
+	TSharedPtr<SProgressBar> DelayProgressBar;
+};
