Copying //UE4/Dev-Sequencer to //UE4/Main

diff --git a/Engine/Source/Editor/SequenceRecorder/Private/SequenceRecorder.cpp b/Engine/Source/Editor/SequenceRecorder/Private/SequenceRecorder.cpp
new file mode 100644
--- /dev/null
+++ b/Engine/Source/Editor/SequenceRecorder/Private/SequenceRecorder.cpp
@@ -0,0 +1,746 @@
+// Copyright 1998-2016 Epic Games, Inc. All Rights Reserved.
+
+#include "SequenceRecorderPrivatePCH.h"
+#include "SequenceRecorder.h"
+#include "ActorRecording.h"
+#include "SequenceRecorderUtils.h"
+#include "SequenceRecorderSettings.h"
+#include "ObjectTools.h"
+#include "ScopedTransaction.h"
+#include "AssetRegistryModule.h"
+#include "SequenceRecorderActorFilter.h"
+#include "Engine/DemoNetDriver.h"
+#include "SNotificationList.h"
+#include "NotificationManager.h"
+#include "DrawDebugHelpers.h"
+#include "LevelSequenceActor.h"
+
+#define LOCTEXT_NAMESPACE "SequenceRecorder"
+
+FSequenceRecorder::FSequenceRecorder()
+	: bQueuedRecordingsDirty(false)
+	, CurrentDelay(0.0f)
+	, CurrentTime(0.0f)
+{
+}
+
+void FSequenceRecorder::Initialize()
+{
+	// load textures we use for the countdown/recording display
+	CountdownTexture = LoadObject<UTexture2D>(nullptr, TEXT("/Engine/EditorResources/SequenceRecorder/Countdown.Countdown"), nullptr, LOAD_None, nullptr);
+	CountdownTexture->AddToRoot();
+	RecordingIndicatorTexture = LoadObject<UTexture2D>(nullptr, TEXT("/Engine/EditorResources/SequenceRecorder/RecordingIndicator.RecordingIndicator"), nullptr, LOAD_None, nullptr);
+	RecordingIndicatorTexture->AddToRoot();
+}
+
+void FSequenceRecorder::Shutdown()
+{
+	if(CountdownTexture.IsValid())
+	{
+		CountdownTexture->RemoveFromRoot();
+		CountdownTexture.Reset();
+	}
+	if(RecordingIndicatorTexture.IsValid())
+	{
+		RecordingIndicatorTexture->RemoveFromRoot();
+		RecordingIndicatorTexture.Reset();
+	}
+}
+
+FSequenceRecorder& FSequenceRecorder::Get()
+{
+	static FSequenceRecorder SequenceRecorder;
+	return SequenceRecorder;
+}
+
+
+bool FSequenceRecorder::IsRecordingQueued(AActor* Actor) const
+{
+	for (UActorRecording* QueuedRecording : QueuedRecordings)
+	{
+		if (QueuedRecording->ActorToRecord == Actor)
+		{
+			return true;
+		}
+	}
+
+	return false;
+}
+
+UActorRecording* FSequenceRecorder::FindRecording(AActor* Actor) const
+{
+	for (UActorRecording* QueuedRecording : QueuedRecordings)
+	{
+		if (QueuedRecording->ActorToRecord == Actor)
+		{
+			return QueuedRecording;
+		}
+	}
+
+	return nullptr;
+}
+
+void FSequenceRecorder::StartAllQueuedRecordings()
+{
+	for (UActorRecording* QueuedRecording : QueuedRecordings)
+	{
+		QueuedRecording->StartRecording(CurrentSequence.Get(), CurrentTime);
+	}
+}
+
+void FSequenceRecorder::StopAllQueuedRecordings()
+{
+	for (UActorRecording* QueuedRecording : QueuedRecordings)
+	{
+		QueuedRecording->StopRecording(CurrentSequence.Get());
+	}
+}
+
+UActorRecording* FSequenceRecorder::AddNewQueuedRecording(AActor* Actor, UAnimSequence* AnimSequence, float Length)
+{
+	const USequenceRecorderSettings* Settings = GetDefault<USequenceRecorderSettings>();
+
+	UActorRecording* ActorRecording = NewObject<UActorRecording>();
+	ActorRecording->AddToRoot();
+	ActorRecording->ActorToRecord = Actor;
+	ActorRecording->TargetAnimation = AnimSequence;
+	ActorRecording->AnimationSettings.Length = Length;
+
+	// We always record in world space as we need animations to record root motion
+	ActorRecording->AnimationSettings.bRecordInWorldSpace = true;
+	ActorRecording->ActorSettings.bRecordTransforms = true;
+
+	// auto-save assets in non-editor runtime
+	if(GEditor == nullptr)
+	{
+		ActorRecording->AnimationSettings.bAutoSaveAsset = true;
+	}
+
+	QueuedRecordings.Add(ActorRecording);
+
+	bQueuedRecordingsDirty = true;
+
+	return ActorRecording;
+}
+
+void FSequenceRecorder::RemoveQueuedRecording(AActor* Actor)
+{
+	for (UActorRecording* QueuedRecording : QueuedRecordings)
+	{
+		if (QueuedRecording->ActorToRecord == Actor)
+		{
+			QueuedRecording->RemoveFromRoot();
+			QueuedRecordings.Remove(QueuedRecording);
+			break;
+		}
+	}
+
+	bQueuedRecordingsDirty = true;
+}
+
+void FSequenceRecorder::RemoveQueuedRecording(class UActorRecording* Recording)
+{
+	for (UActorRecording* QueuedRecording : QueuedRecordings)
+	{
+		if (QueuedRecording == Recording)
+		{
+			QueuedRecording->RemoveFromRoot();
+			QueuedRecordings.Remove(QueuedRecording);
+			break;
+		}
+	}
+
+	bQueuedRecordingsDirty = true;
+}
+
+void FSequenceRecorder::ClearQueuedRecordings()
+{
+	if(IsRecording())
+	{
+		UE_LOG(LogAnimation, Display, TEXT("Couldnt clear queued recordings while recording is in progress"));
+	}
+	else
+	{
+		for (UActorRecording* QueuedRecording : QueuedRecordings)
+		{
+			QueuedRecording->RemoveFromRoot();
+		}
+		QueuedRecordings.Empty();
+
+		bQueuedRecordingsDirty = true;
+	}
+}
+
+bool FSequenceRecorder::HasQueuedRecordings() const
+{
+	return QueuedRecordings.Num() > 0;
+}
+
+
+void FSequenceRecorder::Tick(float DeltaSeconds)
+{
+	const float FirstFrameTickLimit = 1.0f / 30.0f;
+
+	// in-editor we can get a long frame update because of the searching we need to do to filter actors
+	if(DeltaSeconds > FirstFrameTickLimit && CurrentTime < FirstFrameTickLimit * 2.0f && IsRecording())
+	{
+		DeltaSeconds = FirstFrameTickLimit;
+	}
+
+	// if a replay recording is in progress and channels are paused, wait until we have data again before recording
+	if(CurrentReplayWorld.IsValid() && CurrentReplayWorld->DemoNetDriver != nullptr)
+	{
+		if(CurrentReplayWorld->DemoNetDriver->bChannelsArePaused)
+		{
+			return;
+		}
+	}
+
+	const USequenceRecorderSettings* Settings = GetDefault<USequenceRecorderSettings>();
+
+	// check for spawned actors and if they have begun playing yet
+	for(TWeakObjectPtr<AActor>& QueuedSpawnedActor : QueuedSpawnedActors)
+	{
+		if(AActor* Actor = QueuedSpawnedActor.Get())
+		{
+			if(Actor->HasActorBegunPlay())
+			{
+				if(UActorRecording::IsRelevantForRecording(Actor) && IsActorValidForRecording(Actor))
+				{
+					UActorRecording* NewRecording = AddNewQueuedRecording(Actor);
+					NewRecording->bWasSpawnedPostRecord = true;
+					NewRecording->StartRecording(CurrentSequence.Get(), CurrentTime);
+				}
+
+				QueuedSpawnedActor.Reset();
+			}
+		}
+	}
+
+	QueuedSpawnedActors.RemoveAll([](TWeakObjectPtr<AActor>& QueuedSpawnedActor) { return !QueuedSpawnedActor.IsValid(); } );
+
+	FAnimationRecorderManager::Get().Tick(DeltaSeconds);
+
+	for(UActorRecording* Recording : QueuedRecordings)
+	{
+		Recording->Tick(DeltaSeconds, CurrentSequence.Get(), CurrentTime);
+	}
+
+	if(CurrentDelay > 0.0f)
+	{
+		CurrentDelay -= DeltaSeconds;
+		if(CurrentDelay <= 0.0f)
+		{
+			CurrentDelay = 0.0f;
+			StartRecordingInternal(nullptr);
+		}
+	}
+
+	if(Settings->bCreateLevelSequence && CurrentSequence.IsValid())
+	{
+		CurrentTime += DeltaSeconds;
+
+		// check if all our actor recordings are finished or we timed out
+		if(QueuedRecordings.Num() > 0)
+		{
+			bool bAllFinished = true;
+			for(UActorRecording* Recording : QueuedRecordings)
+			{
+				if(Recording->IsRecording())
+				{
+					bAllFinished = false;
+					break;
+				}
+			}
+
+			if(bAllFinished || (Settings->SequenceLength > 0.0f && CurrentTime >= Settings->SequenceLength))
+			{
+				StopRecording();
+			}
+		}
+
+		auto RemoveDeadActorPredicate = 
+			[&](UActorRecording* Recording)
+			{
+				if(!Recording->ActorToRecord.IsValid())
+				{
+					DeadRecordings.Add(Recording);
+					return true;
+				}
+				return false; 
+			};
+
+		// remove all dead actors
+		int32 Removed = QueuedRecordings.RemoveAll(RemoveDeadActorPredicate);
+		if(Removed > 0)
+		{
+			bQueuedRecordingsDirty = true;
+		}
+
+		UpdateSequencePlaybackRange();
+	}
+}
+
+void FSequenceRecorder::DrawDebug(UCanvas* InCanvas, APlayerController* InPlayerController)
+{
+	const float NumFrames = 4.0f;
+	const bool bCountingDown = CurrentDelay > 0.0f && CurrentDelay < NumFrames;
+
+	if(bCountingDown)
+	{
+		const FVector2D IconSize(128.0f, 128.0f);
+		const FVector2D HalfIconSize(64.0f, 64.0f);
+		const float LineThickness = 2.0f;
+
+		FVector2D Center;
+		InCanvas->GetCenter(Center.X, Center.Y);
+		FVector2D IconPosition = Center - HalfIconSize;
+
+		FCanvasIcon Icon = UCanvas::MakeIcon(CountdownTexture.Get(), FMath::FloorToFloat(NumFrames - CurrentDelay) * IconSize.X, 0.0f, IconSize.X, IconSize.Y);
+		InCanvas->DrawIcon(Icon, IconPosition.X, IconPosition.Y);
+		
+		// draw 'clock' line
+		const float Angle = 2.0f * PI * FMath::Fmod(CurrentDelay, 1.0f);
+		const FVector2D AxisX(0.f, -1.f);
+		const FVector2D AxisY(-1.f, 0.f);
+		const FVector2D EndPos = Center + (AxisX * FMath::Cos(Angle) + AxisY * FMath::Sin(Angle)) * (InCanvas->SizeX + InCanvas->SizeY);
+		FCanvasLineItem LineItem(Center, EndPos);
+		LineItem.LineThickness = LineThickness;
+		LineItem.SetColor(FLinearColor::Black);
+		InCanvas->DrawItem(LineItem);
+
+		// draw 'crosshairs'
+		LineItem.Origin = FVector(0.0f, Center.Y, 0.0f);
+		LineItem.EndPos = FVector(InCanvas->SizeX, Center.Y, 0.0f);
+		InCanvas->DrawItem(LineItem);
+
+		LineItem.Origin = FVector(Center.X, 0.0f, 0.0f);
+		LineItem.EndPos = FVector(Center.X, InCanvas->SizeY, 0.0f);
+		InCanvas->DrawItem(LineItem);
+	}
+
+	if(bCountingDown || IsRecording())
+	{
+		FText Text;
+		if(IsRecording())
+		{
+			Text = FText::Format(LOCTEXT("RecordingIndicator", "Recording Sequence: {0}\nTime: {1}"), FText::FromName(CurrentSequence.Get()->GetFName()), FText::AsNumber(CurrentTime));
+		}
+		else
+		{
+			Text = LOCTEXT("RecordingIndicator", "Sequence Recording Pending");
+		}
+
+		const FVector2D IconSize(32.0f, 32.0f);
+		const FVector2D Offset(32.0f, 64.0f);
+
+		InCanvas->SetDrawColor(FColor::White);
+
+		FVector2D IconPosition(Offset.X, Offset.Y * 2.0f);
+		FCanvasIcon Icon = UCanvas::MakeIcon(RecordingIndicatorTexture.Get(), FMath::FloorToFloat(NumFrames - CurrentDelay) * IconSize.X, 0.0f, IconSize.X, IconSize.Y);
+		InCanvas->DrawIcon(Icon, IconPosition.X, IconPosition.Y);
+
+		FFontRenderInfo Info;
+		Info.bEnableShadow = true;
+		InCanvas->DrawText(GEngine->GetMediumFont(), Text, IconPosition.X + IconSize.X + 4.0f, IconPosition.Y + 2.0f, 1.0f, 1.0f, Info);
+	}
+}
+
+bool FSequenceRecorder::StartRecording(const FOnRecordingStarted& OnRecordingStarted, const FOnRecordingFinished& OnRecordingFinished, const FString& InPathToRecordTo, const FString& InSequenceName)
+{
+	OnRecordingStartedDelegate = OnRecordingStarted;
+	OnRecordingFinishedDelegate = OnRecordingFinished;
+
+	const USequenceRecorderSettings* Settings = GetDefault<USequenceRecorderSettings>();
+
+	if(InPathToRecordTo.Len() > 0)
+	{
+		PathToRecordTo = InPathToRecordTo;
+	}
+	else
+	{
+		PathToRecordTo = Settings->SequenceRecordingBasePath.Path;
+	}
+
+	if(InSequenceName.Len() > 0)
+	{
+		SequenceName = InSequenceName;
+	}
+	else
+	{
+		SequenceName = Settings->SequenceName.Len() > 0 ? Settings->SequenceName : TEXT("RecordedSequence");
+	}
+
+	CurrentTime = 0.0f;
+
+	if(Settings->RecordingDelay > 0.0f)
+	{
+		CurrentDelay = Settings->RecordingDelay;
+		
+		UE_LOG(LogAnimation, Display, TEXT("Starting sequence recording with delay of %f seconds"), CurrentDelay);
+
+		return QueuedRecordings.Num() > 0;
+	}
+	
+	return StartRecordingInternal(nullptr);
+}
+
+bool FSequenceRecorder::StartRecordingForReplay(UWorld* World, const FSequenceRecorderActorFilter& ActorFilter)
+{
+	// set up our recording settings
+	USequenceRecorderSettings* Settings = GetMutableDefault<USequenceRecorderSettings>();
+	
+	Settings->bCreateLevelSequence = true;
+	Settings->SequenceLength = 0.0f;
+	Settings->RecordingDelay = 0.0f;
+	Settings->bRecordNearbySpawnedActors = true;
+	Settings->NearbyActorRecordingProximity = 0.0f;
+	Settings->bRecordWorldSettingsActor = true;
+	Settings->ActorFilter = ActorFilter;
+
+	CurrentReplayWorld = World;
+	
+	return StartRecordingInternal(World);
+}
+
+bool FSequenceRecorder::StartRecordingInternal(UWorld* World)
+{
+	CurrentTime = 0.0f;
+
+	const USequenceRecorderSettings* Settings = GetDefault<USequenceRecorderSettings>();
+
+	UWorld* ActorWorld = nullptr;
+	if(Settings->bRecordWorldSettingsActor)
+	{
+		if(World != nullptr || (QueuedRecordings.Num() > 0 && QueuedRecordings[0]->ActorToRecord.IsValid()))
+		{
+			ActorWorld = World != nullptr ? World : QueuedRecordings[0]->ActorToRecord->GetWorld();
+			if(ActorWorld)
+			{
+				AWorldSettings* WorldSettings = ActorWorld->GetWorldSettings();
+				AddNewQueuedRecording(WorldSettings);
+			}
+		}
+	}
+
+	// kick off level sequence actors we are syncing to
+	for(TLazyObjectPtr<ALevelSequenceActor> LevelSequenceActor : Settings->LevelSequenceActorsToTrigger)
+	{
+		if(LevelSequenceActor.IsValid())
+		{
+			// find a counterpart in the PIE world if this actor is not
+			ALevelSequenceActor* ActorToTrigger = LevelSequenceActor.Get();
+			if(ActorWorld && ActorWorld->IsPlayInEditor() && !ActorToTrigger->GetWorld()->IsPlayInEditor())
+			{
+				ActorToTrigger = Cast<ALevelSequenceActor>(EditorUtilities::GetSimWorldCounterpartActor(ActorToTrigger));
+			}
+
+			if(ActorToTrigger && ActorToTrigger->SequencePlayer)
+			{
+				ActorToTrigger->SequencePlayer->Play();
+			}
+		}
+	}
+
+	if(QueuedRecordings.Num() > 0)
+	{
+		ULevelSequence* LevelSequence = nullptr;
+		
+		if(Settings->bCreateLevelSequence)
+		{
+			LevelSequence = SequenceRecorderUtils::MakeNewAsset<ULevelSequence>(PathToRecordTo, SequenceName);
+			if(LevelSequence)
+			{
+				LevelSequence->Initialize();
+				CurrentSequence = LevelSequence;
+
+				FAssetRegistryModule::AssetCreated(LevelSequence);
+			}
+		}
+
+		// register for spawning delegate in the world(s) of recorded actors
+		for(UActorRecording* Recording : QueuedRecordings)
+		{
+			if(Recording->ActorToRecord.IsValid())
+			{
+				UWorld* ActorToRecordWorld = Recording->ActorToRecord->GetWorld();
+				if(ActorToRecordWorld != nullptr)
+				{
+					FDelegateHandle* FoundHandle = ActorSpawningDelegateHandles.Find(ActorToRecordWorld);
+					if(FoundHandle == nullptr)
+					{
+						FDelegateHandle NewHandle = ActorToRecordWorld->AddOnActorSpawnedHandler(FOnActorSpawned::FDelegate::CreateStatic(&FSequenceRecorder::HandleActorSpawned));
+						ActorSpawningDelegateHandles.Add(ActorToRecordWorld, NewHandle);
+					}
+				}
+			}
+		}
+
+		// start recording 
+		bool bStartedRecordingAllActors = true;
+		for(UActorRecording* Recording : QueuedRecordings)
+		{
+			if(!Recording->StartRecording(CurrentSequence.Get(), CurrentTime))
+			{
+				bStartedRecordingAllActors = false;
+				break;
+			}
+		}
+
+		if(!bStartedRecordingAllActors)
+		{
+			// if we couldnt start a recording, stop all others
+			TArray<FAssetData> AssetsToCleanUp;
+			if(LevelSequence)
+			{
+				AssetsToCleanUp.Add(LevelSequence);
+			}
+
+			for(UActorRecording* Recording : QueuedRecordings)
+			{
+				Recording->StopRecording(CurrentSequence.Get());
+			}
+
+			// clean up any assets that we can
+			if(AssetsToCleanUp.Num() > 0)
+			{
+				ObjectTools::DeleteAssets(AssetsToCleanUp, false);
+			}
+		}
+
+#if WITH_EDITOR
+		// if recording via PIE, be sure to stop recording cleanly when PIE ends
+		if (ActorWorld && ActorWorld->IsPlayInEditor())
+		{
+			FEditorDelegates::EndPIE.AddRaw(this, &FSequenceRecorder::HandleEndPIE);
+		}
+#endif
+
+		UE_LOG(LogAnimation, Display, TEXT("Started recording sequence %s"), *LevelSequence->GetPathName());
+
+		OnRecordingStartedDelegate.ExecuteIfBound(CurrentSequence.Get());
+		return true;
+	}
+	else
+	{
+		UE_LOG(LogAnimation, Display, TEXT("No recordings queued, aborting recording"));
+	}
+
+	return false;
+} 
+
+void FSequenceRecorder::HandleEndPIE(bool bSimulating)
+{
+	StopRecording();
+
+#if WITH_EDITOR
+	FEditorDelegates::EndPIE.RemoveAll(this);
+#endif
+}
+
+bool FSequenceRecorder::StopRecording()
+{
+	FScopedSlowTask SlowTask((float)(QueuedRecordings.Num() + DeadRecordings.Num()), LOCTEXT("ProcessingRecording", "Processing Recording"));
+	SlowTask.MakeDialog(false, true);
+
+	CurrentDelay = 0.0f;
+
+	// remove spawn delegates
+	for(auto It = ActorSpawningDelegateHandles.CreateConstIterator(); It; ++It)
+	{
+		TWeakObjectPtr<UWorld> World = It->Key;
+		if(World.IsValid())
+		{
+			World->RemoveOnActorSpawnedHandler(It->Value);
+		}
+	}
+
+	ActorSpawningDelegateHandles.Empty();
+
+	// also stop all dead animation recordings, i.e. ones that use GCed components
+	const bool bShowMessage = false;
+	FAnimationRecorderManager::Get().StopRecordingDeadAnimations(bShowMessage);
+
+	for(UActorRecording* Recording : QueuedRecordings)
+	{
+		SlowTask.EnterProgressFrame();
+
+		Recording->StopRecording(CurrentSequence.Get());
+	}
+
+	for(UActorRecording* Recording : DeadRecordings)
+	{
+		SlowTask.EnterProgressFrame();
+
+		Recording->StopRecording(CurrentSequence.Get());
+	}
+
+	DeadRecordings.Empty();
+
+	const USequenceRecorderSettings* Settings = GetDefault<USequenceRecorderSettings>();
+	if(Settings->bCreateLevelSequence)
+	{
+		if(CurrentSequence.IsValid())
+		{
+			ULevelSequence* LevelSequence = CurrentSequence.Get();
+
+			// Stop referencing the sequence so we are listed as 'not recording'
+			CurrentSequence = nullptr;
+
+			// set movie scene playback range to encompass all sections
+			UpdateSequencePlaybackRange();
+
+			if(GEditor == nullptr)
+			{
+				// auto-save asset outside of the editor
+				UPackage* const Package = LevelSequence->GetOutermost();
+				FString const PackageName = Package->GetName();
+				FString const PackageFileName = FPackageName::LongPackageNameToFilename(PackageName, FPackageName::GetAssetPackageExtension());
+
+				UPackage::SavePackage(Package, NULL, RF_Standalone, *PackageFileName, GError, nullptr, false, true, SAVE_NoError);
+			}
+
+			if(FSlateApplication::IsInitialized())
+			{
+				const FText NotificationText = FText::Format(LOCTEXT("RecordSequence", "'{0}' has been successfully recorded."), FText::FromString(LevelSequence->GetName()));
+					
+				FNotificationInfo Info(NotificationText);
+				Info.ExpireDuration = 8.0f;
+				Info.bUseLargeFont = false;
+				Info.Hyperlink = FSimpleDelegate::CreateLambda([=]()
+				{
+					TArray<UObject*> Assets;
+					Assets.Add(LevelSequence);
+					FAssetEditorManager::Get().OpenEditorForAssets(Assets);
+				});
+				Info.HyperlinkText = FText::Format(LOCTEXT("OpenNewAnimationHyperlink", "Open {0}"), FText::FromString(LevelSequence->GetName()));
+				TSharedPtr<SNotificationItem> Notification = FSlateNotificationManager::Get().AddNotification(Info);
+				if ( Notification.IsValid() )
+				{
+					Notification->SetCompletionState( SNotificationItem::CS_Success );
+				}
+			}
+
+			UE_LOG(LogAnimation, Display, TEXT("Stopped recording sequence %s"), *LevelSequence->GetPathName());
+
+			OnRecordingFinishedDelegate.ExecuteIfBound(LevelSequence);
+
+			OnRecordingFinishedDelegate = FOnRecordingFinished();
+			OnRecordingStartedDelegate = FOnRecordingStarted();
+
+			return true;
+		}
+	}
+	else
+	{
+		UE_LOG(LogAnimation, Display, TEXT("Stopped recording, no sequence created"));
+		return true;
+	}
+
+	return false;
+}
+
+void FSequenceRecorder::UpdateSequencePlaybackRange()
+{
+	if(CurrentSequence.IsValid())
+	{
+		UMovieScene* MovieScene = CurrentSequence->GetMovieScene();
+		if(MovieScene)
+		{
+			float MaxRange = 0.0f;
+			float MinRange = 0.0f;
+
+			TArray<UMovieSceneSection*> MovieSceneSections = MovieScene->GetAllSections();
+			for(UMovieSceneSection* Section : MovieSceneSections)
+			{
+				MaxRange = FMath::Max(MaxRange, Section->GetEndTime());
+				MinRange = FMath::Min(MinRange, Section->GetStartTime());
+			}
+
+			MovieScene->SetPlaybackRange(MinRange, MaxRange);
+		}
+	}
+}
+
+bool FSequenceRecorder::IsDelaying() const
+{
+	return CurrentDelay > 0.0f;
+}
+
+float FSequenceRecorder::GetCurrentDelay() const
+{
+	return CurrentDelay;
+}
+
+bool FSequenceRecorder::IsActorValidForRecording(AActor* Actor)
+{
+	const USequenceRecorderSettings* Settings = GetDefault<USequenceRecorderSettings>();
+
+	float Distance = Settings->NearbyActorRecordingProximity;
+
+	// check distance if valid
+	if(Settings->bRecordNearbySpawnedActors && Distance > 0.0f)
+	{
+		const FTransform ActorTransform = Actor->GetTransform();
+		const FVector ActorTranslation = ActorTransform.GetTranslation();
+
+		for(UActorRecording* Recording : QueuedRecordings)
+		{
+			if(AActor* OtherActor = Recording->ActorToRecord.Get())
+			{
+				if(OtherActor != Actor)
+				{
+					const FTransform OtherActorTransform = OtherActor->GetTransform();
+					const FVector OtherActorTranslation = OtherActorTransform.GetTranslation();
+
+					if((OtherActorTranslation - ActorTranslation).Size() < Distance)
+					{
+						return true;
+					}
+				}
+			}
+		}
+	}
+
+	// check class if any
+	for(const TSubclassOf<AActor>& ActorClass : Settings->ActorFilter.ActorClassesToRecord)
+	{
+		if(Actor->IsA(*ActorClass))
+		{
+			return true;
+		}
+	}
+
+	return false;
+}
+
+void FSequenceRecorder::HandleActorSpawned(AActor* Actor)
+{
+	const USequenceRecorderSettings* Settings = GetDefault<USequenceRecorderSettings>();
+
+	if(Actor && FSequenceRecorder::Get().IsRecording())
+	{
+		// Queue a possible recording - we need to wait until the actor has begun playing
+		FSequenceRecorder::Get().QueuedSpawnedActors.Add(Actor);
+	}
+}
+
+void FSequenceRecorder::HandleActorDespawned(AActor* Actor)
+{
+	if(Actor && FSequenceRecorder::Get().IsRecording())
+	{
+		for(int32 Index = 0; Index < QueuedRecordings.Num(); ++Index)
+		{
+			UActorRecording* Recording = QueuedRecordings[Index];
+			if(Recording->ActorToRecord.Get() == Actor)
+			{
+				Recording->InvalidateObjectToRecord();
+				DeadRecordings.Add(Recording);
+				QueuedRecordings.RemoveAt(Index);
+				break;
+			}
+		}
+	}
+}
+
+#undef LOCTEXT_NAMESPACE
\ No newline at end of file
diff --git a/Engine/Source/Editor/SequenceRecorder/Private/SequenceRecorder.h b/Engine/Source/Editor/SequenceRecorder/Private/SequenceRecorder.h
new file mode 100644
--- /dev/null
+++ b/Engine/Source/Editor/SequenceRecorder/Private/SequenceRecorder.h
@@ -0,0 +1,133 @@
+// Copyright 1998-2015 Epic Games, Inc. All Rights Reserved.
+
+#pragma once
+
+#include "ISequenceRecorder.h"
+
+struct FSequenceRecorder
+{
+public:
+	/** Singleton accessor */
+	static FSequenceRecorder& Get();
+
+	/** Initialize any resources we need */
+	void Initialize();
+
+	/** Clear any resources we need */
+	void Shutdown();
+
+	/** Starts recording a sequence. */
+	bool StartRecording(const FOnRecordingStarted& OnRecordingStarted = FOnRecordingStarted(), const FOnRecordingFinished& OnRecordingFinished = FOnRecordingFinished(), const FString& InPathToRecordTo = FString(), const FString& InSequenceName = FString());
+
+	/** Starts recording a sequence for the specified world. */
+	bool StartRecordingForReplay(UWorld* World, const struct FSequenceRecorderActorFilter& ActorFilter);
+
+	/** Stops any currently recording sequence */
+	bool StopRecording();
+
+	/** Tick the sequence recorder */
+	void Tick(float DeltaSeconds);
+
+	/** Get whether we are currently delaying a recording */
+	bool IsDelaying() const;
+
+	/** Get the current delay we are waiting for */
+	float GetCurrentDelay() const;
+
+	TWeakObjectPtr<class ULevelSequence> GetCurrentSequence() { return CurrentSequence; }
+
+	bool IsRecordingQueued(AActor* Actor) const;
+
+	class UActorRecording* FindRecording(AActor* Actor) const;
+
+	void StartAllQueuedRecordings();
+
+	void StopAllQueuedRecordings();
+
+	void StopRecordingDeadAnimations();
+
+	class UActorRecording* AddNewQueuedRecording(AActor* Actor = nullptr, UAnimSequence* AnimSequence = nullptr, float Length = 0.0f);
+
+	void RemoveQueuedRecording(AActor* Actor);
+
+	void RemoveQueuedRecording(class UActorRecording* Recording);
+
+	bool HasQueuedRecordings() const;
+
+	void ClearQueuedRecordings();
+
+	const TArray<class UActorRecording*>& GetQueuedRecordings() { return QueuedRecordings; }
+
+	bool AreQueuedRecordingsDirty() const { return bQueuedRecordingsDirty; }
+
+	void ResetQueuedRecordingsDirty() { bQueuedRecordingsDirty = false; }
+
+	bool IsRecording() const { return CurrentSequence.IsValid(); }
+
+	/** Draw the countdown to the screen */
+	void DrawDebug(UCanvas* InCanvas, APlayerController* InPlayerController);
+
+	/** Handle actors being spawned */
+	static void HandleActorSpawned(AActor* Actor);
+
+	/** Handle actors being de-spawned */
+	void HandleActorDespawned(AActor* Actor);
+
+private:
+	/** Starts recording a sequence, possibly delayed */
+	bool StartRecordingInternal(UWorld* World);
+
+	/** Check if an actor is valid for recording */
+	bool IsActorValidForRecording(AActor* Actor);
+
+	/** Handle exiting cleanly from PIE */
+	void HandleEndPIE(bool bSimulating);
+
+	/** Keep sequence range up to date with sections that are being recorded */
+	void UpdateSequencePlaybackRange();
+
+private:
+	/** Constructor, private - use Get() function */
+	FSequenceRecorder();
+
+	/** Currently recording level sequence, if any */
+	TWeakObjectPtr<class ULevelSequence> CurrentSequence;
+
+	/** World we are recording a replay for, if any */
+	TLazyObjectPtr<class UWorld> CurrentReplayWorld;
+
+	TArray<class UActorRecording*> QueuedRecordings;
+
+	TArray<class UActorRecording*> DeadRecordings;
+
+	TArray<TWeakObjectPtr<AActor>> QueuedSpawnedActors;
+
+	bool bQueuedRecordingsDirty;
+
+	/** The delay we are currently waiting for */
+	float CurrentDelay;
+
+	/** Current recording time */
+	float CurrentTime;
+
+	/** Delegate handles for FOnActorSpawned events */
+	TMap<TWeakObjectPtr<UWorld>, FDelegateHandle> ActorSpawningDelegateHandles;
+
+	/** Texture we use for the countdown */
+	TWeakObjectPtr<UTexture> CountdownTexture;
+
+	/** Texture we use for the recording indicator */
+	TWeakObjectPtr<UTexture> RecordingIndicatorTexture;
+
+	/** Delegate fired when recording is started */
+	FOnRecordingStarted OnRecordingStartedDelegate;
+
+	/** Delegate fired when recording has finished */
+	FOnRecordingFinished OnRecordingFinishedDelegate;
+
+	/** Cached sequence name to record to */
+	FString SequenceName;
+
+	/** Cached sequence path to record to */
+	FString PathToRecordTo;
+};
diff --git a/Engine/Source/Editor/SequenceRecorder/Private/SequenceRecorderCommands.cpp b/Engine/Source/Editor/SequenceRecorder/Private/SequenceRecorderCommands.cpp
new file mode 100644
--- /dev/null
+++ b/Engine/Source/Editor/SequenceRecorder/Private/SequenceRecorderCommands.cpp
@@ -0,0 +1,21 @@
+// Copyright 1998-2015 Epic Games, Inc. All Rights Reserved.
+
+#include "SequenceRecorderPrivatePCH.h"
+#include "SequenceRecorderCommands.h"
+
+#define LOCTEXT_NAMESPACE "SequenceRecorderCommands"
+
+FSequenceRecorderCommands::FSequenceRecorderCommands()
+	: TCommands<FSequenceRecorderCommands>("SequenceRecorder.Common", LOCTEXT("Common", "Common"), NAME_None, FEditorStyle::GetStyleSetName())
+{
+}
+
+void FSequenceRecorderCommands::RegisterCommands()
+{
+	UI_COMMAND(RecordAll, "Record", "Record all recordings", EUserInterfaceActionType::Button, FInputChord(EModifierKey::Control | EModifierKey::Shift, EKeys::R));
+	UI_COMMAND(StopAll, "StopAll", "Stop all recordings", EUserInterfaceActionType::Button, FInputChord(EModifierKey::Control | EModifierKey::Shift, EKeys::S));
+	UI_COMMAND(AddRecording, "Add", "Add a new recording", EUserInterfaceActionType::Button, FInputChord());
+	UI_COMMAND(RemoveRecording, "Remove", "Remove selected recording", EUserInterfaceActionType::Button, FInputChord());
+}
+
+#undef LOCTEXT_NAMESPACE
\ No newline at end of file
diff --git a/Engine/Source/Editor/SequenceRecorder/Private/SequenceRecorderModule.cpp b/Engine/Source/Editor/SequenceRecorder/Private/SequenceRecorderModule.cpp
new file mode 100644
--- /dev/null
+++ b/Engine/Source/Editor/SequenceRecorder/Private/SequenceRecorderModule.cpp
@@ -0,0 +1,698 @@
+// Copyright 1998-2016 Epic Games, Inc. All Rights Reserved.
+
+#include "SequenceRecorderPrivatePCH.h"
+#include "ModuleManager.h"
+#include "SSequenceRecorder.h"
+#include "SDockTab.h"
+#include "SequenceRecorderCommands.h"
+#include "SequenceRecorder.h"
+#include "SequenceRecorderSettings.h"
+#include "ISequenceRecorder.h"
+#include "NotificationManager.h"
+#include "SNotificationList.h"
+
+#define LOCTEXT_NAMESPACE "SequenceRecorder"
+
+static const FName SequenceRecorderTabName(TEXT("SequenceRecorder"));
+
+static TAutoConsoleVariable<float> CVarDefaultRecordedAnimLength(
+	TEXT("AnimRecorder.AnimLength"),
+	FAnimationRecordingSettings::DefaultMaximumLength,
+	TEXT("Sets default animation length for the animation recorder system."),
+	ECVF_Default);
+
+static TAutoConsoleVariable<float> CVarAnimRecorderSampleRate(
+	TEXT("AnimRecorder.SampleRate"),
+	FAnimationRecordingSettings::DefaultSampleRate,
+	TEXT("Sets the sample rate for the animation recorder system"),
+	ECVF_Default);
+
+static TAutoConsoleVariable<int32> CVarAnimRecorderWorldSpace(
+	TEXT("AnimRecorder.RecordInWorldSpace"),
+	1,
+	TEXT("True to record anim keys in world space, false to record only in local space."),
+	ECVF_Default);
+
+class FSequenceRecorderModule : public ISequenceRecorder, private FSelfRegisteringExec
+{
+	/** IModuleInterface implementation */
+	virtual void StartupModule() override
+	{
+#if WITH_EDITOR
+		GetMutableDefault<USequenceRecorderSettings>()->LoadConfig();
+
+		// set cvar defaults
+		CVarDefaultRecordedAnimLength.AsVariable()->SetOnChangedCallback(FConsoleVariableDelegate::CreateLambda([](IConsoleVariable* Variable)
+		{
+			GetMutableDefault<USequenceRecorderSettings>()->DefaultAnimationSettings.Length = CVarDefaultRecordedAnimLength.GetValueOnGameThread();
+		}));
+
+		CVarAnimRecorderSampleRate.AsVariable()->SetOnChangedCallback(FConsoleVariableDelegate::CreateLambda([](IConsoleVariable* Variable)
+		{
+			GetMutableDefault<USequenceRecorderSettings>()->DefaultAnimationSettings.SampleRate = CVarAnimRecorderSampleRate.GetValueOnGameThread();
+		}));
+
+		CVarAnimRecorderWorldSpace.AsVariable()->SetOnChangedCallback(FConsoleVariableDelegate::CreateLambda([](IConsoleVariable* Variable)
+		{
+			GetMutableDefault<USequenceRecorderSettings>()->DefaultAnimationSettings.bRecordInWorldSpace = (CVarAnimRecorderWorldSpace.GetValueOnGameThread() != 0);
+		}));
+
+		FSequenceRecorderCommands::Register();
+
+		// init sequence recorder
+		FSequenceRecorder::Get().Initialize();
+
+		// register main tick
+		if(GEngine)
+		{
+			PostEditorTickHandle = GEngine->OnPostEditorTick().AddStatic(&FSequenceRecorderModule::TickSequenceRecorder);
+		}
+
+		if (GEditor)
+		{
+			// register Persona recorder
+			FPersonaModule& PersonaModule = FModuleManager::LoadModuleChecked<FPersonaModule>(TEXT("Persona"));
+			PersonaModule.OnIsRecordingActive().BindStatic(&FSequenceRecorderModule::HandlePersonaIsRecordingActive);
+			PersonaModule.OnRecord().BindStatic(&FSequenceRecorderModule::HandlePersonaRecord);
+			PersonaModule.OnStopRecording().BindStatic(&FSequenceRecorderModule::HandlePersonaStopRecording);
+			PersonaModule.OnGetCurrentRecording().BindStatic(&FSequenceRecorderModule::HandlePersonaCurrentRecording);
+			PersonaModule.OnGetCurrentRecordingTime().BindStatic(&FSequenceRecorderModule::HandlePersonaCurrentRecordingTime);
+
+			// register 'keep simulation changes' recorder
+			FLevelEditorModule& LevelEditorModule = FModuleManager::LoadModuleChecked<FLevelEditorModule>(TEXT("LevelEditor"));
+			LevelEditorModule.OnCaptureSingleFrameAnimSequence().BindStatic(&FSequenceRecorderModule::HandleCaptureSingleFrameAnimSequence);
+
+			// register level editor menu extender
+			LevelEditorMenuExtenderDelegate = FLevelEditorModule::FLevelViewportMenuExtender_SelectedActors::CreateStatic(&FSequenceRecorderModule::ExtendLevelViewportContextMenu);
+			auto& MenuExtenders = LevelEditorModule.GetAllLevelViewportContextMenuExtenders();
+			MenuExtenders.Add(LevelEditorMenuExtenderDelegate);
+			LevelEditorExtenderDelegateHandle = MenuExtenders.Last().GetHandle();
+
+			// register standalone UI
+			FGlobalTabmanager::Get()->RegisterNomadTabSpawner(SequenceRecorderTabName, FOnSpawnTab::CreateStatic(&FSequenceRecorderModule::SpawnSequenceRecorderTab))
+				.SetGroup(WorkspaceMenu::GetMenuStructure().GetLevelEditorCategory())
+				.SetDisplayName(LOCTEXT("SequenceRecorderTabTitle", "Sequence Recorder"))
+				.SetTooltipText(LOCTEXT("SequenceRecorderTooltipText", "Open the Sequence Recorder tab."))
+				.SetIcon(FSlateIcon(FEditorStyle::GetStyleSetName(), "SequenceRecorder.TabIcon"));
+
+			// register for debug drawing
+			DrawDebugDelegateHandle = UDebugDrawService::Register(TEXT("Game"),  FDebugDrawDelegate::CreateStatic(&FSequenceRecorderModule::DrawDebug));
+		}
+#endif
+	}
+
+	virtual void ShutdownModule() override
+	{
+#if WITH_EDITOR
+
+		FSequenceRecorder::Get().Shutdown();
+
+		if (GEditor)
+		{
+			UDebugDrawService::Unregister(DrawDebugDelegateHandle);
+
+			if (FSlateApplication::IsInitialized())
+			{
+				FGlobalTabmanager::Get()->UnregisterTabSpawner(SequenceRecorderTabName);
+			}
+
+			if(FModuleManager::Get().IsModuleLoaded(TEXT("LevelEditor")))
+			{
+				FLevelEditorModule& LevelEditorModule = FModuleManager::GetModuleChecked<FLevelEditorModule>(TEXT("LevelEditor"));
+				LevelEditorModule.OnCaptureSingleFrameAnimSequence().Unbind();
+				LevelEditorModule.GetAllLevelViewportContextMenuExtenders().RemoveAll([&](const FLevelEditorModule::FLevelViewportMenuExtender_SelectedActors& Delegate) {
+					return Delegate.GetHandle() == LevelEditorExtenderDelegateHandle;
+				});
+			}
+
+			if (FModuleManager::Get().IsModuleLoaded(TEXT("Persona")))
+			{
+				FPersonaModule& PersonaModule = FModuleManager::GetModuleChecked<FPersonaModule>(TEXT("Persona"));
+				PersonaModule.OnIsRecordingActive().Unbind();
+				PersonaModule.OnRecord().Unbind();
+				PersonaModule.OnStopRecording().Unbind();
+				PersonaModule.OnGetCurrentRecording().Unbind();
+				PersonaModule.OnGetCurrentRecordingTime().Unbind();
+			}
+		}
+
+		if(GEngine)
+		{
+			GEngine->OnPostEditorTick().Remove(PostEditorTickHandle);
+		}
+#endif
+	}
+
+	// FSelfRegisteringExec implementation
+	virtual bool Exec(UWorld* InWorld, const TCHAR* Cmd, FOutputDevice& Ar) override
+	{
+#if WITH_EDITOR
+		if (FParse::Command(&Cmd, TEXT("RecordAnimation")))
+		{
+			return HandleRecordAnimationCommand(InWorld, Cmd, Ar);
+		}
+		else if (FParse::Command(&Cmd, TEXT("StopRecordingAnimation")))
+		{
+			return HandleStopRecordAnimationCommand(InWorld, Cmd, Ar);
+		}
+		else if (FParse::Command(&Cmd, TEXT("RecordSequence")))
+		{
+			return HandleRecordSequenceCommand(InWorld, Cmd, Ar);
+		}
+		else if (FParse::Command(&Cmd, TEXT("StopRecordingSequence")))
+		{
+			return HandleStopRecordSequenceCommand(InWorld, Cmd, Ar);
+		}
+#endif
+		return false;
+	}
+
+	static bool HandleRecordAnimationCommand(UWorld* InWorld, const TCHAR* InStr, FOutputDevice& Ar)
+	{
+#if WITH_EDITOR
+		const TCHAR* Str = InStr;
+		// parse actor name
+		TCHAR ActorName[128];
+		AActor* FoundActor = nullptr;
+		if (FParse::Token(Str, ActorName, ARRAY_COUNT(ActorName), 0))
+		{
+			FString const ActorNameStr = FString(ActorName);
+			for (ULevel const* Level : InWorld->GetLevels())
+			{
+				if (Level)
+				{
+					for (AActor* Actor : Level->Actors)
+					{
+						if (Actor)
+						{
+							if (Actor->GetName() == ActorNameStr)
+							{
+								FoundActor = Actor;
+								break;
+							}
+						}
+					}
+				}
+			}
+		}
+
+		if (FoundActor)
+		{
+			USkeletalMeshComponent* const SkelComp = FoundActor->FindComponentByClass<USkeletalMeshComponent>();
+			if (SkelComp)
+			{
+				TCHAR AssetPath[256];
+				FParse::Token(Str, AssetPath, ARRAY_COUNT(AssetPath), 0);
+				FString const AssetName = FPackageName::GetLongPackageAssetName(AssetPath);
+				return FAnimationRecorderManager::Get().RecordAnimation(SkelComp, AssetPath, AssetName, GetDefault<USequenceRecorderSettings>()->DefaultAnimationSettings);
+			}
+		}
+#endif
+		return false;
+	}
+
+	static AActor* FindActor(const FString& ActorNameStr, UWorld* InWorld, bool bFuzzy = false)
+	{
+		// search for the actor by name
+		for (ULevel* Level : InWorld->GetLevels())
+		{
+			if (Level)
+			{
+				for (AActor* Actor : Level->Actors)
+				{
+					if (Actor)
+					{
+						if (Actor->GetName() == ActorNameStr)
+						{
+							return Actor;
+						}
+					}
+				}
+			}
+		}
+
+		// if we want to do a fuzzy search then we return the first actor whose name that starts 
+		// the specified string
+		if(bFuzzy)
+		{
+			FString FuzzyActorNameStr = ActorNameStr.ToLower();
+
+			for (ULevel* Level : InWorld->GetLevels())
+			{
+				if (Level)
+				{
+					for (AActor* Actor : Level->Actors)
+					{
+						if (Actor)
+						{
+							if (Actor->GetName().ToLower().StartsWith(FuzzyActorNameStr))
+							{
+								return Actor;
+							}
+						}
+					}
+				}
+			}			
+		}
+
+		return nullptr;
+	}
+
+	static bool HandleStopRecordAnimationCommand(UWorld* InWorld, const TCHAR* InStr, FOutputDevice& Ar)
+	{
+#if WITH_EDITOR
+		const TCHAR* Str = InStr;
+
+		// parse actor name
+		TCHAR ActorName[128];
+		AActor* FoundActor = nullptr;
+		bool bStopAll = false;
+		if (FParse::Token(Str, ActorName, ARRAY_COUNT(ActorName), 0))
+		{
+			FString const ActorNameStr = FString(ActorName);
+
+			if (ActorNameStr.ToLower() == TEXT("all"))
+			{
+				bStopAll = true;
+			}
+			else if (InWorld)
+			{
+				FoundActor = FindActor(ActorNameStr, InWorld);
+			}
+		}
+
+		if (bStopAll)
+		{
+			FAnimationRecorderManager::Get().StopRecordingAllAnimations();
+			return true;
+		}
+		else if (FoundActor)
+		{
+			USkeletalMeshComponent* const SkelComp = FoundActor->FindComponentByClass<USkeletalMeshComponent>();
+			if (SkelComp)
+			{
+				FAnimationRecorderManager::Get().StopRecordingAnimation(SkelComp);
+				return true;
+			}
+		}
+
+#endif
+		return false;
+	}
+
+	static void FindActorsOfClass(UClass* Class, UWorld* InWorld, TArray<AActor*>& OutActors)
+	{
+		for (ULevel* Level : InWorld->GetLevels())
+		{
+			if (Level)
+			{
+				for (AActor* Actor : Level->Actors)
+				{
+					if (Actor && Actor->IsA(Class) && UActorRecording::IsRelevantForRecording(Actor))
+					{
+						OutActors.AddUnique(Actor);
+					}
+				}
+			}
+		}
+	}
+
+	static bool HandleRecordSequenceCommand(UWorld* InWorld, const TCHAR* InStr, FOutputDevice& Ar)
+	{
+#if WITH_EDITOR
+		USequenceRecorderSettings* Settings = GetMutableDefault<USequenceRecorderSettings>();
+
+		enum class EFilterType : int32
+		{
+			None,
+			All,
+			Actor,
+			Class
+		};
+
+		const TCHAR* Str = InStr;
+		EFilterType FilterType = EFilterType::None;
+		TCHAR Filter[128];
+		if(FParse::Token(Str, Filter, ARRAY_COUNT(Filter), 0))
+		{
+			FString const FilterStr = FString(Filter).ToLower();
+			if (FilterStr == TEXT("all"))
+			{
+				FilterType = EFilterType::All;
+			}
+			else if(FilterStr == TEXT("actor"))
+			{
+				FilterType = EFilterType::Actor;
+			}
+			else if(FilterStr == TEXT("class"))
+			{
+				FilterType = EFilterType::Class;
+			}
+			else
+			{
+				UE_LOG(LogAnimation, Warning, TEXT("Couldnt parse recording filter, using actor filters from settings."));
+			}
+		}
+
+		if(FilterType == EFilterType::Actor || FilterType == EFilterType::Class)
+		{
+			TCHAR Specifier[128];
+			if(FParse::Token(Str, Specifier, ARRAY_COUNT(Specifier), 0))
+			{
+				FString const SpecifierStr = FString(Specifier).Trim();
+				if(FilterType == EFilterType::Actor)
+				{
+					AActor* FoundActor = FindActor(SpecifierStr, InWorld, true);
+					if(FoundActor)
+					{
+						Settings->ActorFilter.ActorClassesToRecord.Empty();
+						FSequenceRecorder::Get().ClearQueuedRecordings();
+						FSequenceRecorder::Get().AddNewQueuedRecording(FoundActor);
+						FSequenceRecorder::Get().StartRecording();					
+					}
+					return true;
+				}
+				else if(FilterType == EFilterType::Class)
+				{
+					UClass* FoundClass = FindObject<UClass>(ANY_PACKAGE, *SpecifierStr);
+					if(FoundClass != nullptr)
+					{
+						Settings->ActorFilter.ActorClassesToRecord.Empty();
+						Settings->ActorFilter.ActorClassesToRecord.Add(FoundClass);
+						Settings->bRecordNearbySpawnedActors = false;
+						Settings->NearbyActorRecordingProximity = 0.0f;
+
+						FSequenceRecorder::Get().ClearQueuedRecordings();
+
+						TArray<AActor*> ActorsToRecord;
+						FindActorsOfClass(FoundClass, InWorld, ActorsToRecord);
+
+						for(AActor* ActorToRecord : ActorsToRecord)
+						{
+							FSequenceRecorder::Get().AddNewQueuedRecording(ActorToRecord);
+						}
+
+						FSequenceRecorder::Get().StartRecording();
+						return true;
+					}
+					else
+					{
+						UE_LOG(LogAnimation, Warning, TEXT("Couldnt parse class filter, aborting recording."));
+					}
+				}
+			}
+		}
+		else
+		{
+			FSequenceRecorder::Get().ClearQueuedRecordings();
+
+			TArray<AActor*> ActorsToRecord;
+			if(FilterType == EFilterType::None)
+			{
+				for(TSubclassOf<AActor>& SubClass : Settings->ActorFilter.ActorClassesToRecord)
+				{
+					FindActorsOfClass(*SubClass, InWorld, ActorsToRecord);
+				}
+			}
+			else
+			{
+				Settings->bRecordNearbySpawnedActors = false;
+				Settings->NearbyActorRecordingProximity = 0.0f;
+
+				Settings->ActorFilter.ActorClassesToRecord.Empty();
+				Settings->ActorFilter.ActorClassesToRecord.Add(AActor::StaticClass());
+
+				FindActorsOfClass(AActor::StaticClass(), InWorld, ActorsToRecord);
+			}
+
+			for(AActor* ActorToRecord : ActorsToRecord)
+			{
+				FSequenceRecorder::Get().AddNewQueuedRecording(ActorToRecord);
+			}
+
+			FSequenceRecorder::Get().StartRecording();
+			return true;
+		}
+#endif
+		return false;
+	}
+
+	bool HandleStopRecordSequenceCommand(UWorld* InWorld, const TCHAR* InStr, FOutputDevice& Ar)
+	{
+#if WITH_EDITOR
+		FSequenceRecorder::Get().StopRecording();
+		FSequenceRecorder::Get().ClearQueuedRecordings();
+		return true;
+#else
+		return false;
+#endif
+	}
+
+	// ISequenceRecorder interface
+	virtual bool StartRecording(UWorld* World, const FSequenceRecorderActorFilter& ActorFilter) override
+	{
+		return FSequenceRecorder::Get().StartRecordingForReplay(World, ActorFilter);
+	}
+
+	virtual void StopRecording() override
+	{
+		FSequenceRecorder::Get().StopRecording();
+	}
+
+	virtual bool IsRecording() override
+	{
+		return FSequenceRecorder::Get().IsRecording();
+	}
+
+	virtual float GetCurrentRecordingLength() override
+	{
+		TWeakObjectPtr<ULevelSequence> CurrentSequence = FSequenceRecorder::Get().GetCurrentSequence();
+
+		return CurrentSequence.IsValid() ? CurrentSequence->GetMovieScene()->GetPlaybackRange().Size<float>() : 0.0f;
+	}
+
+	virtual bool StartRecording(UWorld* World, const FString& ActorNameToRecord, const FOnRecordingStarted& OnRecordingStarted, const FOnRecordingFinished& OnRecordingFinished, const FString& PathToRecordTo, const FString& SequenceName) override
+	{
+		if(ActorNameToRecord.Len() > 0)
+		{
+			AActor* Actor = FindActor(ActorNameToRecord, World, true);
+			if(Actor != nullptr)
+			{
+				FSequenceRecorder::Get().ClearQueuedRecordings();
+				FSequenceRecorder::Get().AddNewQueuedRecording(Actor);
+			}
+			else
+			{
+				if(FSlateApplication::IsInitialized())
+				{
+					FNotificationInfo Info(FText::Format(LOCTEXT("SequenceRecordingErrorActor", "Couldnt find an actor to record matching name \"{0}\""), FText::FromString(ActorNameToRecord)));
+					Info.bUseLargeFont = false;
+
+					FSlateNotificationManager::Get().AddNotification(Info);
+				}
+
+				UE_LOG(LogAnimation, Display, TEXT("Couldnt find an actor to record matching name \"%s\""), *ActorNameToRecord);
+			}
+		}
+
+		return FSequenceRecorder::Get().StartRecording(OnRecordingStarted, OnRecordingFinished, PathToRecordTo, SequenceName);
+	}
+
+	virtual void NotifyActorStartRecording(AActor* Actor)
+	{
+		FSequenceRecorder::HandleActorSpawned(Actor);
+	}
+
+	virtual void NotifyActorStopRecording(AActor* Actor)
+	{
+		FSequenceRecorder::Get().HandleActorDespawned(Actor);
+	}
+
+	static void TickSequenceRecorder(float DeltaSeconds)
+	{
+		if (!IsRunningDedicatedServer() && !IsRunningCommandlet())
+		{
+			FSequenceRecorder::Get().Tick(DeltaSeconds);
+		}
+	}
+
+#if WITH_EDITOR
+	static UAnimSequence* HandleCaptureSingleFrameAnimSequence(USkeletalMeshComponent* Component)
+	{
+		FAnimationRecorder Recorder;
+		if (Recorder.TriggerRecordAnimation(Component))
+		{
+			class UAnimSequence * Sequence = Recorder.GetAnimationObject();
+			if (Sequence)
+			{
+				Recorder.StopRecord(false);
+				return Sequence;
+			}
+		}
+
+		return nullptr;
+	}
+
+	static void HandlePersonaIsRecordingActive(USkeletalMeshComponent* Component, bool& bIsRecording)
+	{
+		bIsRecording = FAnimationRecorderManager::Get().IsRecording(Component);
+	}
+
+	static void HandlePersonaRecord(USkeletalMeshComponent* Component)
+	{
+		FAnimationRecorderManager::Get().RecordAnimation(Component);
+	}
+
+	static void HandlePersonaStopRecording(USkeletalMeshComponent* Component)
+	{
+		FAnimationRecorderManager::Get().StopRecordingAnimation(Component);
+	}
+
+	static void HandlePersonaTickRecording(USkeletalMeshComponent* Component, float DeltaSeconds)
+	{
+	//	FAnimationRecorderManager::Get().Tick(Component, DeltaSeconds);
+	}
+
+	static void HandlePersonaCurrentRecording(USkeletalMeshComponent* Component, UAnimSequence*& OutSequence)
+	{
+		OutSequence = FAnimationRecorderManager::Get().GetCurrentlyRecordingSequence(Component);
+	}
+
+	static void HandlePersonaCurrentRecordingTime(USkeletalMeshComponent* Component, float& OutTime)
+	{
+		OutTime = FAnimationRecorderManager::Get().GetCurrentRecordingTime(Component);
+	}
+
+	static TSharedRef<SDockTab> SpawnSequenceRecorderTab(const FSpawnTabArgs& SpawnTabArgs)
+	{
+		const TSharedRef<SDockTab> MajorTab =
+			SNew(SDockTab)
+			.Icon(FEditorStyle::Get().GetBrush("SequenceRecorder.TabIcon"))
+			.TabRole(ETabRole::NomadTab);
+
+		MajorTab->SetContent(SNew(SSequenceRecorder));
+
+		return MajorTab;
+	}
+
+	static TSharedRef<FExtender> ExtendLevelViewportContextMenu(const TSharedRef<FUICommandList> CommandList, const TArray<AActor*> SelectedActors)
+	{
+		TSharedRef<FExtender> Extender(new FExtender());
+
+		// check if actors are already registered for recording
+		bool bCanAddRecording = false;
+		bool bCanRemoveRecording = false;
+
+		FSequenceRecorder& SequenceRecorder = FSequenceRecorder::Get();
+
+		for (AActor* Actor : SelectedActors)
+		{
+			const bool bIsQueued = SequenceRecorder.IsRecordingQueued(Actor);
+			bCanRemoveRecording |= bIsQueued;
+			bCanAddRecording |= !bIsQueued;
+		}
+
+		if (bCanAddRecording || bCanRemoveRecording)
+		{
+			// Add the sequence recorder sub-menu extender
+			Extender->AddMenuExtension(
+				"ActorSelectVisibilityLevels",
+				EExtensionHook::After,
+				nullptr,
+				FMenuExtensionDelegate::CreateStatic(&FSequenceRecorderModule::CreateLevelViewportContextMenuEntries, bCanAddRecording, bCanRemoveRecording));
+		}
+
+		return Extender;
+	}
+
+	static void CreateLevelViewportContextMenuEntries(FMenuBuilder& MenuBuilder, bool bCanAddRecording, bool bCanRemoveRecording)
+	{
+		MenuBuilder.BeginSection("SequenceRecording", LOCTEXT("SequenceRecordingLevelEditorHeading", "Sequence Recording"));
+
+		if (bCanAddRecording)
+		{
+			FUIAction Action_AddQuickRecording(FExecuteAction::CreateStatic(&FSequenceRecorderModule::AddActorsForQuickRecording));
+
+			MenuBuilder.AddMenuEntry(
+				LOCTEXT("MenuExtensionAddQuickRecording", "Trigger Actor Recording"),
+				LOCTEXT("MenuExtensionAddQuickRecording_Tooltip", "Set up the selected actors for recording and trigger recording immediately"),
+				FSlateIcon(),
+				Action_AddQuickRecording,
+				NAME_None,
+				EUserInterfaceActionType::Button);
+
+			FUIAction Action_AddRecording(FExecuteAction::CreateStatic(&FSequenceRecorderModule::AddActorsForRecording));
+
+			MenuBuilder.AddMenuEntry(
+				LOCTEXT("MenuExtensionAddRecording", "Queue Actor Recording"),
+				LOCTEXT("MenuExtensionAddRecording_Tooltip", "Queue up the selected actors for recording"),
+				FSlateIcon(),
+				Action_AddRecording,
+				NAME_None,
+				EUserInterfaceActionType::Button);
+		}
+
+		if (bCanRemoveRecording)
+		{
+			FUIAction Action_RemoveRecording(FExecuteAction::CreateStatic(&FSequenceRecorderModule::RemoveActorsForRecording));
+
+			MenuBuilder.AddMenuEntry(
+				LOCTEXT("MenuExtensionRemoveRecording", "Removed Queued Recording"),
+				LOCTEXT("MenuExtensionRemoveRecording_Tooltip", "Remove the selected actors from the recording queue"),
+				FSlateIcon(),
+				Action_RemoveRecording,
+				NAME_None,
+				EUserInterfaceActionType::Button);
+		}
+
+		MenuBuilder.EndSection();
+	}
+
+	static void AddActorsForRecording()
+	{
+		FSequenceRecorder& SequenceRecorder = FSequenceRecorder::Get();
+
+		TArray<AActor*> SelectedActors;
+		GEditor->GetSelectedActors()->GetSelectedObjects<AActor>(SelectedActors);
+		for(AActor* Actor : SelectedActors)
+		{
+			SequenceRecorder.AddNewQueuedRecording(Actor, nullptr, GetDefault<USequenceRecorderSettings>()->SequenceLength);
+		}
+	}
+
+	static void AddActorsForQuickRecording()
+	{
+		AddActorsForRecording();
+
+		FSequenceRecorder::Get().StartRecording();
+	}
+
+	static void RemoveActorsForRecording()
+	{
+		FSequenceRecorder& SequenceRecorder = FSequenceRecorder::Get();
+
+		TArray<AActor*> SelectedActors;
+		GEditor->GetSelectedActors()->GetSelectedObjects<AActor>(SelectedActors);
+		for (AActor* Actor : SelectedActors)
+		{
+			SequenceRecorder.RemoveQueuedRecording(Actor);
+		}
+	}
+
+	static void DrawDebug(UCanvas* InCanvas, APlayerController* InPlayerController)
+	{
+		FSequenceRecorder::Get().DrawDebug(InCanvas, InPlayerController);
+	}
+#endif
+	FDelegateHandle PostEditorTickHandle;
+
+	FLevelEditorModule::FLevelViewportMenuExtender_SelectedActors LevelEditorMenuExtenderDelegate;
+
+	FDelegateHandle LevelEditorExtenderDelegateHandle;
+
+	FDelegateHandle DrawDebugDelegateHandle;
+};
+
+IMPLEMENT_MODULE( FSequenceRecorderModule, SequenceRecorder )
+
+#undef LOCTEXT_NAMESPACE
\ No newline at end of file
diff --git a/Engine/Source/Editor/SequenceRecorder/Private/SequenceRecorderPrivatePCH.h b/Engine/Source/Editor/SequenceRecorder/Private/SequenceRecorderPrivatePCH.h
new file mode 100644
--- /dev/null
+++ b/Engine/Source/Editor/SequenceRecorder/Private/SequenceRecorderPrivatePCH.h
@@ -0,0 +1,15 @@
+// Copyright 1998-2016 Epic Games, Inc. All Rights Reserved.
+
+#include "Editor.h"
+#include "LevelEditor.h"
+#include "PersonaModule.h"
+#include "Animation/AnimationRecordingSettings.h"
+#include "AnimationRecorder.h"
+#include "SlateBasics.h"
+#include "WorkspaceMenuStructureModule.h"
+#include "PropertyEditing.h"
+#include "MovieScene.h"
+#include "KeyParams.h"
+#include "IKeyframeSection.h"
+#include "MovieSceneSection.h"
+#include "LevelSequence.h"
\ No newline at end of file
diff --git a/Engine/Source/Editor/SequenceRecorder/Private/SequenceRecorderUtils.cpp b/Engine/Source/Editor/SequenceRecorder/Private/SequenceRecorderUtils.cpp
new file mode 100644
--- /dev/null
+++ b/Engine/Source/Editor/SequenceRecorder/Private/SequenceRecorderUtils.cpp
@@ -0,0 +1,26 @@
+// Copyright 1998-2016 Epic Games, Inc. All Rights Reserved.
+
+#include "SequenceRecorderPrivatePCH.h"
+#include "SequenceRecorderUtils.h"
+
+namespace SequenceRecorderUtils
+{
+
+AActor* GetAttachment(AActor* InActor, FName& SocketName, FName& ComponentName)
+{
+	AActor* AttachedActor = nullptr;
+	if(InActor)
+	{
+		USceneComponent* RootComponent = InActor->GetRootComponent();
+		if(RootComponent && RootComponent->AttachParent != nullptr)
+		{
+			AttachedActor = RootComponent->AttachParent->GetOwner();
+			SocketName = RootComponent->AttachSocketName;
+			ComponentName = RootComponent->AttachParent->GetFName();
+		}
+	}
+
+	return AttachedActor;
+}
+
+}
\ No newline at end of file
diff --git a/Engine/Source/Editor/SequenceRecorder/Public/ISequenceRecorder.h b/Engine/Source/Editor/SequenceRecorder/Public/ISequenceRecorder.h
new file mode 100644
--- /dev/null
+++ b/Engine/Source/Editor/SequenceRecorder/Public/ISequenceRecorder.h
@@ -0,0 +1,56 @@
+// Copyright 1998-2016 Epic Games, Inc. All Rights Reserved.
+
+#pragma once
+
+#include "ModuleInterface.h"
+
+DECLARE_DELEGATE_OneParam(FOnRecordingStarted, class UMovieSceneSequence* /*Sequence*/);
+
+DECLARE_DELEGATE_OneParam(FOnRecordingFinished, class UMovieSceneSequence* /*Sequence*/);
+
+class ISequenceRecorder : public IModuleInterface
+{
+public:
+	/** 
+	 * Start recording the passed-in actors.
+	 * @param	World			The world we use to record actors
+	 * @param	ActorFilter		Actor filter to gather actors spawned in this world for recording
+	 * @return true if recording was successfully started
+	 */
+	virtual bool StartRecording(UWorld* World, const struct FSequenceRecorderActorFilter& ActorFilter) = 0;
+
+	/** Stop recording current sequence, if any */
+	virtual void StopRecording() = 0;
+
+	/** Are we currently recording a sequence */
+	virtual bool IsRecording() = 0;
+
+	/** How long is the currently recording sequence */
+	virtual float GetCurrentRecordingLength() = 0;
+
+	/**
+	 * Start a recording, possibly with some delay (specified by the sequence recording settings).
+	 * @param	World				The world we use to record actors
+	 * @param	ActorNameToRecord	The name of the actor to record (fuzzy-matched, precise name not required).
+	 * @param	OnRecordingStarted	Delegate fired when recording has commenced.
+	 * @param	OnRecordingFinished	Delegate fired when recording has finished.
+	 * @param	PathToRecordTo		Optional path to a sequence to record to. If none is specified we use the defaults in the settings.
+	 * @param	SequenceName		Optional name of a sequence to record to. If none is specified we use the defaults in the settings.
+	 * @return true if recording was successfully started
+	*/
+	virtual bool StartRecording(UWorld* World, const FString& ActorNameToRecord, const FOnRecordingStarted& OnRecordingStarted, const FOnRecordingFinished& OnRecordingFinished, const FString& PathToRecordTo = FString(), const FString& SequenceName = FString()) = 0;
+
+	/**
+	 * Notify we should start recording an actor - usually used for 'actor pooling' implementations
+	 * to simulate spawning. Has no effect when recording is not in progress.
+	 * @param	Actor	The actor that was 'spawned'
+	 */
+	virtual void NotifyActorStartRecording(AActor* Actor) = 0;
+
+	/**
+	 * Notify we should stop recording an actor - usually used for 'actor pooling' implementations
+	 * to simulate de-spawning. Has no effect when recording is not in progress.
+	 * @param	Actor	The actor that was 'de-spawned'
+	 */
+	virtual void NotifyActorStopRecording(AActor* Actor) = 0;
+};
\ No newline at end of file
diff --git a/Engine/Source/Editor/SequenceRecorder/Public/SequenceRecorderBlueprintLibrary.h b/Engine/Source/Editor/SequenceRecorder/Public/SequenceRecorderBlueprintLibrary.h
new file mode 100644
--- /dev/null
+++ b/Engine/Source/Editor/SequenceRecorder/Public/SequenceRecorderBlueprintLibrary.h
@@ -0,0 +1,26 @@
+// Copyright 1998-2016 Epic Games, Inc. All Rights Reserved.
+
+#pragma once
+
+#include "SequenceRecorderBlueprintLibrary.generated.h"
+
+UCLASS()
+class SEQUENCERECORDER_API USequenceRecorderBlueprintLibrary : public UBlueprintFunctionLibrary
+{
+	GENERATED_BODY()
+
+	/** 
+	 * Start recording the passed-in actors to a level sequence.
+	 * @param	ActorsToRecord	The actors to record
+	 */
+	UFUNCTION(BlueprintCallable, Category="Sequence Recording")
+	static void StartRecordingSequence(const TArray<AActor*>& ActorsToRecord);
+
+	/** Are we currently recording a sequence */
+	UFUNCTION(BlueprintPure, Category="Sequence Recording")
+	static bool IsRecordingSequence();
+
+	/** Stop recording the current sequence, if any */
+	UFUNCTION(BlueprintCallable, Category="Sequence Recording")
+	static void StopRecordingSequence();
+};
\ No newline at end of file
diff --git a/Engine/Source/Editor/SequenceRecorder/SequenceRecorder.Build.cs b/Engine/Source/Editor/SequenceRecorder/SequenceRecorder.Build.cs
new file mode 100644
--- /dev/null
+++ b/Engine/Source/Editor/SequenceRecorder/SequenceRecorder.Build.cs
@@ -0,0 +1,59 @@
+// Copyright 1998-2016 Epic Games, Inc. All Rights Reserved.
+
+namespace UnrealBuildTool.Rules
+{
+	public class SequenceRecorder : ModuleRules
+	{
+		public SequenceRecorder(TargetInfo Target)
+		{
+			PublicIncludePaths.AddRange(
+				new string[] {
+				}
+				);
+
+			PrivateIncludePaths.AddRange(
+				new string[] {
+					"Editor/SequenceRecorder/Private",
+					"Editor/SequenceRecorder/Private/PropertyRecorders",
+				}
+				);
+
+			PublicDependencyModuleNames.AddRange(
+				new string[]
+				{
+				}
+				);
+
+			PrivateDependencyModuleNames.AddRange(
+				new string[]
+				{
+					"Core",
+					"CoreUObject",
+					"SlateCore",
+					"Slate",
+					"InputCore",
+					"Engine",
+					"UnrealEd",
+					"EditorStyle",
+					"Projects",
+					"LevelEditor",
+					"Persona",
+					"WorkspaceMenuStructure",
+					"PropertyEditor",
+					"MovieScene",
+					"MovieSceneTracks",
+					"LevelSequence",
+					"NetworkReplayStreaming",
+					"AssetRegistry",
+					"CinematicCamera",
+				}
+				);
+
+			DynamicallyLoadedModuleNames.AddRange(
+				new string[]
+				{
+				}
+				);
+		}
+	}
+}
\ No newline at end of file
diff --git a/Engine/Source/Editor/Sequencer/Private/DisplayNodes/SequencerDisplayNode.cpp b/Engine/Source/Editor/Sequencer/Private/DisplayNodes/SequencerDisplayNode.cpp
--- a/Engine/Source/Editor/Sequencer/Private/DisplayNodes/SequencerDisplayNode.cpp
+++ b/Engine/Source/Editor/Sequencer/Private/DisplayNodes/SequencerDisplayNode.cpp
@@ -166,7 +166,7 @@ void FSequencerDisplayNode::Initialize(float InVirtualTop, float InVirtualBottom
 
 void FSequencerDisplayNode::AddObjectBindingNode(TSharedRef<FSequencerObjectBindingNode> ObjectBindingNode)
 {
-	ChildNodes.Add(ObjectBindingNode);
+	AddChildAndSetParent( ObjectBindingNode );
 }
 
 
@@ -300,7 +300,7 @@ TSharedRef<FSequencerTrackNode> FSequencerDisplayNode::AddSectionAreaNode(UMovie
 	if (!SectionNode.IsValid())
 	{
 		// No existing node found make a new one
-		SectionNode = MakeShareable(new FSequencerTrackNode(AssociatedTrack, AssociatedEditor, SharedThis(this), ParentTree));
+		SectionNode = MakeShareable( new FSequencerTrackNode( AssociatedTrack, AssociatedEditor, false, SharedThis(this), ParentTree ) );
 		ChildNodes.Add( SectionNode.ToSharedRef() );
 	}
 
@@ -334,6 +334,15 @@ void FSequencerDisplayNode::AddKeyAreaNode(FName KeyAreaName, const FText& Displ
 	KeyAreaNode->AddKeyArea(KeyArea);
 }
 
+FLinearColor FSequencerDisplayNode::GetDisplayNameColor() const
+{
+	return FLinearColor( 1.f, 1.f, 1.f, 1.f );
+}
+
+FText FSequencerDisplayNode::GetDisplayNameToolTipText() const
+{
+	return FText();
+}
 
 TSharedRef<SWidget> FSequencerDisplayNode::GenerateContainerWidgetForOutliner(const TSharedRef<SSequencerTreeViewRow>& InRow)
 {
@@ -425,7 +434,7 @@ FString FSequencerDisplayNode::GetPathName() const
 }
 
 
-TSharedPtr<SWidget> FSequencerDisplayNode::OnSummonContextMenu(const FGeometry& MyGeometry, const FPointerEvent& MouseEvent)
+TSharedPtr<SWidget> FSequencerDisplayNode::OnSummonContextMenu()
 {
 	// @todo sequencer replace with UI Commands instead of faking it
 	const bool bShouldCloseWindowAfterMenuSelection = true;
@@ -593,4 +602,12 @@ void FSequencerDisplayNode::EnableKeyGoupingRegeneration()
 	KeyGroupRegenerationLock.Decrement();
 }
 
+
+void FSequencerDisplayNode::AddChildAndSetParent( TSharedRef<FSequencerDisplayNode> InChild )
+{
+	ChildNodes.Add( InChild );
+	InChild->ParentNode = SharedThis( this );
+}
+
+
 #undef LOCTEXT_NAMESPACE
diff --git a/Engine/Source/Editor/Sequencer/Private/DisplayNodes/SequencerDisplayNode.h b/Engine/Source/Editor/Sequencer/Private/DisplayNodes/SequencerDisplayNode.h
--- a/Engine/Source/Editor/Sequencer/Private/DisplayNodes/SequencerDisplayNode.h
+++ b/Engine/Source/Editor/Sequencer/Private/DisplayNodes/SequencerDisplayNode.h
@@ -8,6 +8,7 @@ class ISequencerSection;
 class SAnimationOutlinerTreeNode;
 class SSequencerTreeViewRow;
 class FGroupedKeyArea;
+class FSequencerDisplayNodeDragDropOp;
 
 
 /**
@@ -45,6 +46,8 @@ namespace ESequencerNode
 		Category,
 		/* Benign spacer node */
 		Spacer,
+		/* Folder node */
+		Folder
 	};
 }
 
@@ -136,6 +139,16 @@ public:
 	 */
 	virtual FText GetDisplayName() const = 0;
 
+	/**
+	* @return the color used to draw the display name.
+	*/
+	virtual FLinearColor GetDisplayNameColor() const;
+
+	/**
+	 * @return the text to display for the tool tip for the display name. 
+	 */
+	virtual FText GetDisplayNameToolTipText() const;
+
 	/**
 	 * Set the node's display name.
 	 *
@@ -198,7 +211,7 @@ public:
 	FString GetPathName() const;
 
 	/** Summon context menu */
-	TSharedPtr<SWidget> OnSummonContextMenu(const FGeometry& MyGeometry, const FPointerEvent& MouseEvent);
+	TSharedPtr<SWidget> OnSummonContextMenu();
 
 	/** What sort of context menu this node summons */
 	virtual void BuildContextMenu(FMenuBuilder& MenuBuilder);
@@ -351,9 +364,27 @@ public:
 	/** Get the key grouping for the specified section index */
 	TSharedRef<FGroupedKeyArea> GetKeyGrouping(int32 InSectionIndex);
 
+	/** Get the number of key groupings */
+	int32 GetNumKeyGroupings() { return KeyGroupings.Num(); }
+
 	DECLARE_EVENT(FSequencerDisplayNode, FRequestRenameEvent);
 	FRequestRenameEvent& OnRenameRequested() { return RenameRequestedEvent; }
 
+	/** 
+	 * Returns whether or not this node can be dragged. 
+	 */
+	virtual bool CanDrag() const { return false; }
+
+	/**
+	 * Determines if there is a valid drop zone based on the current drag drop operation and the zone the items were dragged onto.
+	 */
+	virtual TOptional<EItemDropZone> CanDrop( FSequencerDisplayNodeDragDropOp& DragDropOp, EItemDropZone ItemDropZone ) const { return TOptional<EItemDropZone>(); }
+
+	/**
+	 * Handles a drop of items onto this display node.
+	 */
+	virtual void Drop( const TArray<TSharedRef<FSequencerDisplayNode>>& DraggedNodes, EItemDropZone DropZone ) { }
+
 public:
 
 	/** Temporarily disable dynamic regeneration of key groupings. This prevents overlapping key groups from being amalgamated during drags. Key times will continue to update correctly. */
@@ -362,6 +393,11 @@ public:
 	/** Re-enable dynamic regeneration of key groupings */
 	static void EnableKeyGoupingRegeneration();
 
+protected:
+
+	/** Adds a child to this node, and sets it's parent to this node. */
+	void AddChildAndSetParent( TSharedRef<FSequencerDisplayNode> InChild );
+
 private:
 
 	/** Callback for executing a "Rename Node" context menu action. */
diff --git a/Engine/Source/Editor/Sequencer/Private/DisplayNodes/SequencerFolderNode.cpp b/Engine/Source/Editor/Sequencer/Private/DisplayNodes/SequencerFolderNode.cpp
new file mode 100644
--- /dev/null
+++ b/Engine/Source/Editor/Sequencer/Private/DisplayNodes/SequencerFolderNode.cpp
@@ -0,0 +1,270 @@
+// Copyright 1998-2016 Epic Games, Inc. All Rights Reserved.
+
+#include "SequencerPrivatePCH.h"
+#include "MovieSceneFolder.h"
+#include "SequencerUtilities.h"
+#include "MovieScene.h"
+#include "MovieSceneSequence.h"
+#include "SequencerDisplayNodeDragDropOp.h"
+#include "ScopedTransaction.h"
+
+FSequencerFolderNode::FSequencerFolderNode( UMovieSceneFolder& InMovieSceneFolder, TSharedPtr<FSequencerDisplayNode> InParentNode, FSequencerNodeTree& InParentTree )
+	: FSequencerDisplayNode( InMovieSceneFolder.GetFolderName(), InParentNode, InParentTree )
+	, MovieSceneFolder(InMovieSceneFolder)
+{
+	FolderOpenBrush = FEditorStyle::GetBrush( "ContentBrowser.AssetTreeFolderOpen" );
+	FolderClosedBrush = FEditorStyle::GetBrush( "ContentBrowser.AssetTreeFolderClosed" );
+}
+
+
+ESequencerNode::Type FSequencerFolderNode::GetType() const
+{
+	return ESequencerNode::Folder;
+}
+
+
+float FSequencerFolderNode::GetNodeHeight() const
+{
+	// TODO: Add another constant.
+	return SequencerLayoutConstants::FolderNodeHeight;
+}
+
+
+FNodePadding FSequencerFolderNode::GetNodePadding() const
+{
+	return FNodePadding(4, 4);
+}
+
+
+bool FSequencerFolderNode::CanRenameNode() const
+{
+	return true;
+}
+
+
+FText FSequencerFolderNode::GetDisplayName() const
+{
+	return FText::FromName( MovieSceneFolder.GetFolderName() );
+}
+
+
+void FSequencerFolderNode::SetDisplayName( const FText& NewDisplayName )
+{
+	if ( MovieSceneFolder.GetFolderName() != FName( *NewDisplayName.ToString() ) )
+	{
+		TArray<FName> SiblingNames;
+		TSharedPtr<FSequencerDisplayNode> ParentNode = GetParent();
+		if ( ParentNode.IsValid() )
+		{
+			for ( TSharedRef<FSequencerDisplayNode> SiblingNode : ParentNode->GetChildNodes() )
+			{
+				if ( SiblingNode != AsShared() )
+				{
+					SiblingNames.Add( SiblingNode->GetNodeName() );
+				}
+			}
+		}
+
+		FName UniqueName = FSequencerUtilities::GetUniqueName( FName( *NewDisplayName.ToString() ), SiblingNames );
+
+		const FScopedTransaction Transaction( NSLOCTEXT( "SequencerFolderNode", "RenameFolder", "Rename folder." ) );
+		MovieSceneFolder.Modify();
+		MovieSceneFolder.SetFolderName( UniqueName );
+	}
+}
+
+
+const FSlateBrush* FSequencerFolderNode::GetIconBrush() const
+{
+	return IsExpanded()
+		? FolderOpenBrush
+		: FolderClosedBrush;
+}
+
+
+bool FSequencerFolderNode::CanDrag() const
+{
+	return true;
+}
+
+
+TOptional<EItemDropZone> FSequencerFolderNode::CanDrop( FSequencerDisplayNodeDragDropOp& DragDropOp, EItemDropZone ItemDropZone ) const
+{
+	DragDropOp.ResetToDefaultToolTip();
+
+	if ( ItemDropZone == EItemDropZone::AboveItem )
+	{
+		if ( GetParent().IsValid() )
+		{
+			// When dropping above, only allow it for root level nodes.
+			return TOptional<EItemDropZone>();
+		}
+		else
+		{
+			// Make sure there are no folder name collisions with the root folders
+			UMovieScene* FocusedMovieScene = GetParentTree().GetSequencer().GetFocusedMovieSceneSequence()->GetMovieScene();
+			TSet<FName> RootFolderNames;
+			for ( UMovieSceneFolder* RootFolder : FocusedMovieScene->GetRootFolders() )
+			{
+				RootFolderNames.Add( RootFolder->GetFolderName() );
+			}
+
+			for ( TSharedRef<FSequencerDisplayNode> DraggedNode : DragDropOp.GetDraggedNodes() )
+			{
+				if ( DraggedNode->GetType() == ESequencerNode::Folder )
+				{
+					TSharedRef<FSequencerFolderNode> DraggedFolder = StaticCastSharedRef<FSequencerFolderNode>( DraggedNode );
+					if ( RootFolderNames.Contains( DraggedFolder->GetFolder().GetFolderName() ) )
+					{
+						DragDropOp.CurrentHoverText = FText::Format(
+							NSLOCTEXT( "SeqeuencerFolderNode", "DuplicateRootFolderDragErrorFormat", "Root folder with name '{0}' already exists." ),
+							FText::FromName( DraggedFolder->GetFolder().GetFolderName() ) );
+						return TOptional<EItemDropZone>();
+					}
+				}
+			}
+		}
+		return TOptional<EItemDropZone>( EItemDropZone::AboveItem );
+	}
+	else
+	{
+		// When dropping onto, don't allow dropping into the same folder, don't allow dropping
+		// parents into children, and don't allow duplicate folder names.
+		TSet<FName> ChildFolderNames;
+		for ( UMovieSceneFolder* ChildFolder : GetFolder().GetChildFolders() )
+		{
+			ChildFolderNames.Add( ChildFolder->GetFolderName() );
+		}
+
+		for ( TSharedRef<FSequencerDisplayNode> DraggedNode : DragDropOp.GetDraggedNodes() )
+		{
+			TSharedPtr<FSequencerDisplayNode> ParentNode = DraggedNode->GetParent();
+			if ( ParentNode.IsValid() && ParentNode.Get() == this )
+			{
+				DragDropOp.CurrentHoverText = NSLOCTEXT( "SeqeuencerFolderNode", "SameParentDragError", "Can't drag a node onto the same parent." );
+				return TOptional<EItemDropZone>();
+			}
+
+			if ( DraggedNode->GetType() == ESequencerNode::Folder )
+			{
+				TSharedRef<FSequencerFolderNode> DraggedFolder = StaticCastSharedRef<FSequencerFolderNode>( DraggedNode );
+				if ( ChildFolderNames.Contains( DraggedFolder->GetFolder().GetFolderName() ) )
+				{
+					DragDropOp.CurrentHoverText = FText::Format(
+						NSLOCTEXT( "SeqeuencerFolderNode", "DuplicateChildFolderDragErrorFormat", "Folder with name '{0}' already exists." ),
+						FText::FromName( DraggedFolder->GetFolder().GetFolderName() ) );
+					return TOptional<EItemDropZone>();
+				}
+			}
+		}
+		TSharedPtr<FSequencerDisplayNode> CurrentNode = SharedThis( ( FSequencerDisplayNode* )this );
+		while ( CurrentNode.IsValid() )
+		{
+			if ( DragDropOp.GetDraggedNodes().Contains( CurrentNode ) )
+			{
+				DragDropOp.CurrentHoverText = NSLOCTEXT( "SeqeuencerFolderNode", "ParentIntoChildDragError", "Can't drag a parent node into one of it's children." );
+				return TOptional<EItemDropZone>();
+			}
+			CurrentNode = CurrentNode->GetParent();
+		}
+		return TOptional<EItemDropZone>( EItemDropZone::OntoItem );
+	}
+}
+
+
+void FSequencerFolderNode::Drop( const TArray<TSharedRef<FSequencerDisplayNode>>& DraggedNodes, EItemDropZone ItemDropZone )
+{
+	const FScopedTransaction Transaction( NSLOCTEXT( "SequencerFolderNode", "MoveIntoFolder", "Move items into folder." ) );
+	GetFolder().SetFlags(RF_Transactional);
+	GetFolder().Modify();
+	for ( TSharedRef<FSequencerDisplayNode> DraggedNode : DraggedNodes )
+	{
+		TSharedPtr<FSequencerDisplayNode> ParentNode = DraggedNode->GetParent();
+		switch ( DraggedNode->GetType() )
+		{
+			case ESequencerNode::Folder:
+			{
+				TSharedRef<FSequencerFolderNode> DraggedFolderNode = StaticCastSharedRef<FSequencerFolderNode>( DraggedNode );
+				UMovieScene* FocusedMovieScene = GetParentTree().GetSequencer().GetFocusedMovieSceneSequence()->GetMovieScene();
+
+				if ( ItemDropZone == EItemDropZone::OntoItem )
+				{
+					GetFolder().AddChildFolder( &DraggedFolderNode->GetFolder() );
+				}
+				else
+				{
+					FocusedMovieScene->Modify();
+					FocusedMovieScene->GetRootFolders().Add( &DraggedFolderNode->GetFolder() );
+				}
+				
+				if ( ParentNode.IsValid() )
+				{
+					checkf( ParentNode->GetType() == ESequencerNode::Folder, TEXT( "Can not remove from unsupported parent node." ) );
+					TSharedPtr<FSequencerFolderNode> ParentFolder = StaticCastSharedPtr<FSequencerFolderNode>( ParentNode );
+					ParentFolder->GetFolder().Modify();
+					ParentFolder->GetFolder().RemoveChildFolder( &DraggedFolderNode->GetFolder() );
+				}
+				else
+				{
+					FocusedMovieScene->Modify();
+					FocusedMovieScene->GetRootFolders().Remove( &DraggedFolderNode->GetFolder() );
+				}
+
+				break;
+			}
+			case ESequencerNode::Track:
+			{
+				TSharedRef<FSequencerTrackNode> DraggedTrackNode = StaticCastSharedRef<FSequencerTrackNode>( DraggedNode );
+				
+				if( ItemDropZone == EItemDropZone::OntoItem )
+				{
+					GetFolder().AddChildMasterTrack( DraggedTrackNode->GetTrack() );
+				}
+				
+				if ( ParentNode.IsValid() )
+				{
+					checkf( ParentNode->GetType() == ESequencerNode::Folder, TEXT( "Can not remove from unsupported parent node." ) );
+					TSharedPtr<FSequencerFolderNode> ParentFolder = StaticCastSharedPtr<FSequencerFolderNode>( ParentNode );
+					ParentFolder->GetFolder().Modify();
+					ParentFolder->GetFolder().RemoveChildMasterTrack( DraggedTrackNode->GetTrack() );
+				}
+
+				break;
+			}
+			case ESequencerNode::Object:
+			{
+				TSharedRef<FSequencerObjectBindingNode> DraggedObjectBindingNode = StaticCastSharedRef<FSequencerObjectBindingNode>( DraggedNode );
+				
+				if ( ItemDropZone == EItemDropZone::OntoItem )
+				{
+					GetFolder().AddChildObjectBinding( DraggedObjectBindingNode->GetObjectBinding() );
+				}
+
+				if ( ParentNode.IsValid() )
+				{
+					checkf( ParentNode->GetType() == ESequencerNode::Folder, TEXT( "Can not remove from unsupported parent node." ) );
+					TSharedPtr<FSequencerFolderNode> ParentFolder = StaticCastSharedPtr<FSequencerFolderNode>( ParentNode );
+					ParentFolder->GetFolder().Modify();
+					ParentFolder->GetFolder().RemoveChildObjectBinding( DraggedObjectBindingNode->GetObjectBinding() );
+				}
+
+				break;
+			}
+		}
+	}
+	SetExpansionState( true );
+	ParentTree.GetSequencer().NotifyMovieSceneDataChanged();
+}
+
+
+void FSequencerFolderNode::AddChildNode( TSharedRef<FSequencerDisplayNode> ChildNode )
+{
+	AddChildAndSetParent( ChildNode );
+}
+
+
+UMovieSceneFolder& FSequencerFolderNode::GetFolder() const
+{
+	return MovieSceneFolder;
+}
+
diff --git a/Engine/Source/Editor/Sequencer/Private/DisplayNodes/SequencerFolderNode.h b/Engine/Source/Editor/Sequencer/Private/DisplayNodes/SequencerFolderNode.h
new file mode 100644
--- /dev/null
+++ b/Engine/Source/Editor/Sequencer/Private/DisplayNodes/SequencerFolderNode.h
@@ -0,0 +1,41 @@
+// Copyright 1998-2016 Epic Games, Inc. All Rights Reserved.
+
+#pragma once
+
+#include "SequencerDisplayNode.h"
+
+/** A sequencer display node representing folders in the outliner. */
+class FSequencerFolderNode : public FSequencerDisplayNode
+{
+public:
+	FSequencerFolderNode( UMovieSceneFolder& InMovieSceneFolder, TSharedPtr<FSequencerDisplayNode> InParentNode, FSequencerNodeTree& InParentTree );
+
+	// FSequencerDisplayNode interface
+	virtual ESequencerNode::Type GetType() const override;
+	virtual float GetNodeHeight() const override;
+	virtual FNodePadding GetNodePadding() const override;
+	virtual bool CanRenameNode() const override;
+	virtual FText GetDisplayName() const override;
+	virtual void SetDisplayName( const FText& NewDisplayName ) override;
+	virtual const FSlateBrush* GetIconBrush() const override;
+	virtual bool CanDrag() const override;
+	virtual TOptional<EItemDropZone> CanDrop( FSequencerDisplayNodeDragDropOp& DragDropOp, EItemDropZone ItemDropZone ) const override;
+	virtual void Drop( const TArray<TSharedRef<FSequencerDisplayNode>>& DraggedNodes, EItemDropZone ItemDropZone ) override;
+
+	/** Adds a child node to this folder node. */
+	void AddChildNode( TSharedRef<FSequencerDisplayNode> ChildNode );
+
+	/** Gets the folder data for this display node. */
+	UMovieSceneFolder& GetFolder() const;
+
+private:
+
+	/** The brush used to draw the icon when this folder is open .*/
+	const FSlateBrush* FolderOpenBrush;
+
+	/** The brush used to draw the icon when this folder is closed. */
+	const FSlateBrush* FolderClosedBrush;
+
+	/** The movie scene folder data which this node represents. */
+	UMovieSceneFolder& MovieSceneFolder;
+};
\ No newline at end of file
diff --git a/Engine/Source/Editor/Sequencer/Private/DisplayNodes/SequencerObjectBindingNode.cpp b/Engine/Source/Editor/Sequencer/Private/DisplayNodes/SequencerObjectBindingNode.cpp
--- a/Engine/Source/Editor/Sequencer/Private/DisplayNodes/SequencerObjectBindingNode.cpp
+++ b/Engine/Source/Editor/Sequencer/Private/DisplayNodes/SequencerObjectBindingNode.cpp
@@ -195,6 +195,23 @@ void FSequencerObjectBindingNode::AddSpawnOwnershipMenu(FMenuBuilder& MenuBuilde
 		NAME_None,
 		EUserInterfaceActionType::ToggleButton
 	);
+
+	MenuBuilder.AddMenuSeparator();
+	MenuBuilder.AddMenuEntry(
+		LOCTEXT("IgnoreOwnership_Label", "Keep Alive Outside Playback Range (In Sequencer)"),
+		LOCTEXT("IgnoreOwnership_Tooltip", "Keeps the spawned object alive when viewing this specific sequence outside of its playback range. Does not apply to runtime evaluation."),
+		FSlateIcon(),
+		FUIAction(
+			FExecuteAction::CreateLambda([=]{
+				Spawnable->SetIgnoreOwnershipInEditor(!Spawnable->ShouldIgnoreOwnershipInEditor());
+				GetSequencer().SetGlobalTimeDirectly(GetSequencer().GetGlobalTime());
+			}),
+			FCanExecuteAction(),
+			FIsActionChecked::CreateLambda([=]{ return Spawnable->ShouldIgnoreOwnershipInEditor(); })
+		),
+		NAME_None,
+		EUserInterfaceActionType::ToggleButton
+	);
 }
 
 bool FSequencerObjectBindingNode::CanRenameNode() const
@@ -240,6 +257,36 @@ FText FSequencerObjectBindingNode::GetDisplayName() const
 	return DefaultDisplayName;
 }
 
+FLinearColor FSequencerObjectBindingNode::GetDisplayNameColor() const
+{
+	FSequencer& Sequencer = ParentTree.GetSequencer();
+	TArray<TWeakObjectPtr<UObject>> RuntimeObjects;
+	Sequencer.GetRuntimeObjects( Sequencer.GetFocusedMovieSceneSequenceInstance(), ObjectBinding, RuntimeObjects );
+	if ( RuntimeObjects.Num() == 0 )
+	{
+		return FLinearColor::Red;
+	}
+	else
+	{
+		return FSequencerDisplayNode::GetDisplayNameColor();
+	}
+}
+
+FText FSequencerObjectBindingNode::GetDisplayNameToolTipText() const
+{
+	FSequencer& Sequencer = ParentTree.GetSequencer();
+	TArray<TWeakObjectPtr<UObject>> RuntimeObjects;
+	Sequencer.GetRuntimeObjects( Sequencer.GetFocusedMovieSceneSequenceInstance(), ObjectBinding, RuntimeObjects );
+	if ( RuntimeObjects.Num() == 0 )
+	{
+		return LOCTEXT("InvalidBoundObjectToolTip", "The object bound to this track is missing.");
+	}
+	else
+	{
+		return FText();
+	}
+}
+
 const FSlateBrush* FSequencerObjectBindingNode::GetIconBrush() const
 {
 	return FClassIconFinder::FindIconForClass(GetClassForObjectBinding());
@@ -297,6 +344,13 @@ void FSequencerObjectBindingNode::SetDisplayName(const FText& NewDisplayName)
 }
 
 
+bool FSequencerObjectBindingNode::CanDrag() const
+{
+	TSharedPtr<FSequencerDisplayNode> ParentNode = GetParent();
+	return ParentNode.IsValid() == false || ParentNode->GetType() != ESequencerNode::Object;
+}
+
+
 /* FSequencerObjectBindingNode implementation
  *****************************************************************************/
 
diff --git a/Engine/Source/Editor/Sequencer/Private/DisplayNodes/SequencerObjectBindingNode.h b/Engine/Source/Editor/Sequencer/Private/DisplayNodes/SequencerObjectBindingNode.h
--- a/Engine/Source/Editor/Sequencer/Private/DisplayNodes/SequencerObjectBindingNode.h
+++ b/Engine/Source/Editor/Sequencer/Private/DisplayNodes/SequencerObjectBindingNode.h
@@ -51,6 +51,8 @@ public:
 	virtual bool CanRenameNode() const override;
 	virtual TSharedRef<SWidget> GetCustomOutlinerContent() override;
 	virtual FText GetDisplayName() const override;
+	virtual FLinearColor GetDisplayNameColor() const override;
+	virtual FText GetDisplayNameToolTipText() const override;
 	virtual const FSlateBrush* GetIconBrush() const override;
 	virtual const FSlateBrush* GetIconOverlayBrush() const override;
 	virtual FText GetIconToolTipText() const override;
@@ -58,6 +60,7 @@ public:
 	virtual FNodePadding GetNodePadding() const override;
 	virtual ESequencerNode::Type GetType() const override;
 	virtual void SetDisplayName(const FText& NewDisplayName) override;
+	virtual bool CanDrag() const override;
 
 protected:
 
diff --git a/Engine/Source/Editor/Sequencer/Private/DisplayNodes/SequencerSectionKeyAreaNode.h b/Engine/Source/Editor/Sequencer/Private/DisplayNodes/SequencerSectionKeyAreaNode.h
--- a/Engine/Source/Editor/Sequencer/Private/DisplayNodes/SequencerSectionKeyAreaNode.h
+++ b/Engine/Source/Editor/Sequencer/Private/DisplayNodes/SequencerSectionKeyAreaNode.h
@@ -59,7 +59,7 @@ public:
 	 * 
 	 * @return All key areas
 	 */
-	TArray<TSharedRef<IKeyArea>> GetAllKeyAreas() const
+	const TArray<TSharedRef<IKeyArea>>& GetAllKeyAreas() const
 	{
 		return KeyAreas;
 	}
diff --git a/Engine/Source/Editor/Sequencer/Private/DisplayNodes/SequencerTrackNode.cpp b/Engine/Source/Editor/Sequencer/Private/DisplayNodes/SequencerTrackNode.cpp
--- a/Engine/Source/Editor/Sequencer/Private/DisplayNodes/SequencerTrackNode.cpp
+++ b/Engine/Source/Editor/Sequencer/Private/DisplayNodes/SequencerTrackNode.cpp
@@ -19,10 +19,11 @@ namespace SequencerNodeConstants
 /* FTrackNode structors
  *****************************************************************************/
 
-FSequencerTrackNode::FSequencerTrackNode(UMovieSceneTrack& InAssociatedTrack, ISequencerTrackEditor& InAssociatedEditor, TSharedPtr<FSequencerDisplayNode> InParentNode, FSequencerNodeTree& InParentTree)
+FSequencerTrackNode::FSequencerTrackNode(UMovieSceneTrack& InAssociatedTrack, ISequencerTrackEditor& InAssociatedEditor, bool bInCanBeDragged, TSharedPtr<FSequencerDisplayNode> InParentNode, FSequencerNodeTree& InParentTree)
 	: FSequencerDisplayNode(InAssociatedTrack.GetTrackName(), InParentNode, InParentTree)
 	, AssociatedEditor(InAssociatedEditor)
 	, AssociatedTrack(&InAssociatedTrack)
+	, bCanBeDragged(bInCanBeDragged)
 { }
 
 
@@ -110,8 +111,13 @@ void FSequencerTrackNode::BuildContextMenu(FMenuBuilder& MenuBuilder)
 
 bool FSequencerTrackNode::CanRenameNode() const
 {
-	UMovieSceneTrack* Track = AssociatedTrack.Get();
-	return (Track != nullptr) && Track->IsA(UMovieSceneNameableTrack::StaticClass());
+	auto NameableTrack = Cast<UMovieSceneNameableTrack>(AssociatedTrack.Get());
+
+	if (NameableTrack != nullptr)
+	{
+		return NameableTrack->CanRename();
+	}
+	return false;
 }
 
 TSharedRef<SWidget> FSequencerTrackNode::GetCustomOutlinerContent()
@@ -236,6 +242,12 @@ const FSlateBrush* FSequencerTrackNode::GetIconBrush() const
 }
 
 
+bool FSequencerTrackNode::CanDrag() const
+{
+	return bCanBeDragged;
+}
+
+
 void FSequencerTrackNode::GetChildKeyAreaNodesRecursively(TArray<TSharedRef<FSequencerSectionKeyAreaNode>>& OutNodes) const
 {
 	FSequencerDisplayNode::GetChildKeyAreaNodesRecursively(OutNodes);
@@ -255,17 +267,18 @@ FText FSequencerTrackNode::GetDisplayName() const
 
 float FSequencerTrackNode::GetNodeHeight() const
 {
-	float Height = (Sections.Num() > 0)
-		? Sections[0]->GetSectionHeight() * (GetMaxRowIndex() + 1)
-		: SequencerLayoutConstants::SectionAreaDefaultHeight;
+	if (Sections.Num() > 0)
+	{
+		return (Sections[0]->GetSectionHeight() + 2 * SequencerNodeConstants::CommonPadding) * (GetMaxRowIndex() + 1);
+	}
 
-	return Height + 2*SequencerNodeConstants::CommonPadding;
+	return SequencerLayoutConstants::SectionAreaDefaultHeight + 2*SequencerNodeConstants::CommonPadding;
 }
 
 
 FNodePadding FSequencerTrackNode::GetNodePadding() const
 {
-	return FNodePadding(0.f);//SequencerNodeConstants::CommonPadding);
+	return FNodePadding(0.f);
 }
 
 
diff --git a/Engine/Source/Editor/Sequencer/Private/DisplayNodes/SequencerTrackNode.h b/Engine/Source/Editor/Sequencer/Private/DisplayNodes/SequencerTrackNode.h
--- a/Engine/Source/Editor/Sequencer/Private/DisplayNodes/SequencerTrackNode.h
+++ b/Engine/Source/Editor/Sequencer/Private/DisplayNodes/SequencerTrackNode.h
@@ -23,10 +23,11 @@ public:
 	 * 
 	 * @param InAssociatedType The track that this node represents.
 	 * @param InAssociatedEditor The track editor for the track that this node represents.
+	 * @param bInCanBeDragged Whether or not this node can be dragged and dropped.
 	 * @param InParentNode The parent of this node, or nullptr if this is a root node.
 	 * @param InParentTree The tree this node is in.
 	 */
-	FSequencerTrackNode(UMovieSceneTrack& InAssociatedTrack, ISequencerTrackEditor& InAssociatedEditor, TSharedPtr<FSequencerDisplayNode> InParentNode, FSequencerNodeTree& InParentTree);
+	FSequencerTrackNode(UMovieSceneTrack& InAssociatedTrack, ISequencerTrackEditor& InAssociatedEditor, bool bInCanBeDragged, TSharedPtr<FSequencerDisplayNode> InParentNode, FSequencerNodeTree& InParentTree);
 
 public:
 
@@ -67,7 +68,7 @@ public:
 	}
 
 	/** @return Returns the top level key node for the section area if it exists */
-	TSharedPtr<FSequencerSectionKeyAreaNode> GetTopLevelKeyNode()
+	TSharedPtr<FSequencerSectionKeyAreaNode> GetTopLevelKeyNode() const
 	{
 		return TopLevelKeyNode;
 	}
@@ -98,6 +99,7 @@ public:
 	virtual ESequencerNode::Type GetType() const override;
 	virtual void SetDisplayName(const FText& NewDisplayName) override;
 	virtual const FSlateBrush* GetIconBrush() const override;
+	virtual bool CanDrag() const override;
 
 private:
 
@@ -112,4 +114,7 @@ private:
 
 	/** If the section area is a key area itself, this represents the node for the keys */
 	TSharedPtr<FSequencerSectionKeyAreaNode> TopLevelKeyNode;
+
+	/** Whether or not this track node can be dragged. */
+	bool bCanBeDragged;
 };
diff --git a/Engine/Source/Editor/Sequencer/Private/GroupedKeyArea.cpp b/Engine/Source/Editor/Sequencer/Private/GroupedKeyArea.cpp
--- a/Engine/Source/Editor/Sequencer/Private/GroupedKeyArea.cpp
+++ b/Engine/Source/Editor/Sequencer/Private/GroupedKeyArea.cpp
@@ -24,35 +24,119 @@ namespace
 	TMap<FIndexKey, FIndexEntry> GlobalIndex;
 }
 
+FGroupedKeyCollection::FGroupedKeyCollection()
+	: IndexKey(FString(), nullptr)
+{
+	GroupingThreshold = SMALL_NUMBER;
+}
 
-FGroupedKeyArea::FGroupedKeyArea(FSequencerDisplayNode& InNode, int32 InSectionIndex)
-	: Section(nullptr)
-	, IndexKey(FString(), nullptr)
+void FGroupedKeyCollection::InitializeExplicit(const TArray<FSequencerDisplayNode*>& InNodes, float DuplicateThreshold)
 {
+	KeyAreas.Reset();
+	Groups.Reset();
+
 	TArray<TSharedRef<FSequencerSectionKeyAreaNode>> AllKeyAreaNodes;
 	AllKeyAreaNodes.Reserve(36);
-	InNode.GetChildKeyAreaNodesRecursively(AllKeyAreaNodes);
+	for (const FSequencerDisplayNode* Node : InNodes)
+	{
+		const FSequencerSectionKeyAreaNode* KeyAreaNode = nullptr;
 
-	for (const auto& Node : AllKeyAreaNodes)
+		if (Node->GetType() == ESequencerNode::KeyArea)
+		{
+			KeyAreaNode = static_cast<const FSequencerSectionKeyAreaNode*>(Node);
+		}
+		else if (Node->GetType() == ESequencerNode::Track)
+		{
+			KeyAreaNode = static_cast<const FSequencerTrackNode*>(Node)->GetTopLevelKeyNode().Get();
+		}
+
+		if (!KeyAreaNode)
+		{
+			continue;
+		}
+
+		const TArray<TSharedRef<IKeyArea>>& AllKeyAreas = KeyAreaNode->GetAllKeyAreas();
+		KeyAreas.Reserve(KeyAreas.Num() + AllKeyAreas.Num());
+
+		for (const TSharedRef<IKeyArea>& KeyArea : AllKeyAreas)
+		{
+			AddKeyArea(KeyArea);
+		}
+	}
+
+	RemoveDuplicateKeys(DuplicateThreshold);
+}
+
+void FGroupedKeyCollection::InitializeRecursive(const TArray<FSequencerDisplayNode*>& InNodes, float DuplicateThreshold)
+{
+	KeyAreas.Reset();
+	Groups.Reset();
+
+	TArray<TSharedRef<FSequencerSectionKeyAreaNode>> AllKeyAreaNodes;
+	AllKeyAreaNodes.Reserve(36);
+	for (FSequencerDisplayNode* Node : InNodes)
 	{
-		TSharedRef<IKeyArea> KeyArea = Node->GetKeyArea(InSectionIndex);
+		if (Node->GetType() == ESequencerNode::KeyArea)
+		{
+			AllKeyAreaNodes.Add(StaticCastSharedRef<FSequencerSectionKeyAreaNode>(Node->AsShared()));
+		}
 
-		const int32 KeyAreaIndex = KeyAreas.Num();
-		KeyAreas.Add(KeyArea);
+		Node->GetChildKeyAreaNodesRecursively(AllKeyAreaNodes);
+	}
+
+	for (const auto& Node : AllKeyAreaNodes)
+	{
+		const TArray<TSharedRef<IKeyArea>>& AllKeyAreas = Node->GetAllKeyAreas();
+		KeyAreas.Reserve(KeyAreas.Num() + AllKeyAreas.Num());
 
+		for (const TSharedRef<IKeyArea>& KeyArea : AllKeyAreas)
 		{
-			auto* OwningSection = KeyArea->GetOwningSection();
-			// Ensure they all belong to the same section
-			ensure(!Section || Section == OwningSection);
-			Section = OwningSection;
+			AddKeyArea(KeyArea);
 		}
+	}
+
+	RemoveDuplicateKeys(DuplicateThreshold);
+}
 
-		for (const FKeyHandle& KeyHandle : KeyArea->GetUnsortedKeyHandles())
+void FGroupedKeyCollection::InitializeRecursive(FSequencerDisplayNode& InNode, int32 InSectionIndex, float DuplicateThreshold)
+{
+	KeyAreas.Reset();
+	Groups.Reset();
+
+	TArray<TSharedRef<FSequencerSectionKeyAreaNode>> AllKeyAreaNodes;
+	AllKeyAreaNodes.Reserve(36);
+	InNode.GetChildKeyAreaNodesRecursively(AllKeyAreaNodes);
+
+	for (const auto& Node : AllKeyAreaNodes)
+	{
+		TSharedPtr<IKeyArea> KeyArea = Node->GetKeyArea(InSectionIndex);
+		if (KeyArea.IsValid())
 		{
-			Groups.Emplace(KeyArea->GetKeyTime(KeyHandle), KeyAreaIndex, KeyHandle);
+			AddKeyArea(KeyArea.ToSharedRef());
 		}
 	}
 
+	RemoveDuplicateKeys(DuplicateThreshold);
+}
+
+void FGroupedKeyCollection::AddKeyArea(const TSharedRef<IKeyArea>& InKeyArea)
+{
+	const int32 KeyAreaIndex = KeyAreas.Num();
+	KeyAreas.Add(InKeyArea);
+
+	TArray<FKeyHandle> AllKeyHandles = InKeyArea->GetUnsortedKeyHandles();
+	Groups.Reserve(Groups.Num() + AllKeyHandles.Num());
+
+	for (const FKeyHandle& KeyHandle : AllKeyHandles)
+	{
+		Groups.Emplace(InKeyArea->GetKeyTime(KeyHandle), KeyAreaIndex, KeyHandle);
+	}
+}
+
+void FGroupedKeyCollection::RemoveDuplicateKeys(float DuplicateThreshold)
+{
+	GroupingThreshold = DuplicateThreshold;
+
 	Groups.Sort([](const FKeyGrouping& A, const FKeyGrouping& B){
 		return A.RepresentativeTime < B.RepresentativeTime;
 	});
@@ -65,7 +149,7 @@ FGroupedKeyArea::FGroupedKeyArea(FSequencerDisplayNode& InNode, int32 InSectionI
 		int32 NumToMerge = 0;
 		for (int32 DuplIndex = Index + 1; DuplIndex < Groups.Num(); ++DuplIndex)
 		{
-			if (!FMath::IsNearlyEqual(CurrentTime, Groups[DuplIndex].RepresentativeTime))
+			if (!FMath::IsNearlyEqual(CurrentTime, Groups[DuplIndex].RepresentativeTime, DuplicateThreshold))
 			{
 				break;
 			}
@@ -83,15 +167,9 @@ FGroupedKeyArea::FGroupedKeyArea(FSequencerDisplayNode& InNode, int32 InSectionI
 			Groups.RemoveAt(Start, NumToMerge, false);
 		}
 	}
-
-	if (Section)
-	{
-		IndexKey = FIndexKey(InNode.GetPathName(), Section);
-		UpdateIndex();
-	}
 }
 
-void FGroupedKeyArea::UpdateIndex() const
+void FGroupedKeyCollection::UpdateIndex() const
 {
 	auto& IndexEntry = GlobalIndex.FindOrAdd(IndexKey);
 
@@ -131,7 +209,64 @@ void FGroupedKeyArea::UpdateIndex() const
 	IndexEntry.RepresentativeTimes = MoveTemp(NewRepresentativeTimes);
 }
 
-const FKeyGrouping* FGroupedKeyArea::FindGroup(FKeyHandle InHandle) const
+void FGroupedKeyCollection::IterateKeys(const TFunctionRef<bool(float)>& Iter)
+{
+	for (const FKeyGrouping& Grouping : Groups)
+	{
+		if (!Iter(Grouping.RepresentativeTime))
+		{
+			return;
+		}
+	}
+}
+
+float FGroupedKeyCollection::GetKeyGroupingThreshold() const
+{
+	return GroupingThreshold;
+}
+
+TOptional<float> FGroupedKeyCollection::FindFirstKeyInRange(const TRange<float>& InRange, EFindKeyDirection Direction) const
+{
+	// @todo: linear search may be slow where there are lots of keys
+
+	bool bWithinRange = false;
+	if (Direction == EFindKeyDirection::Backwards)
+	{
+		for (int32 Index = Groups.Num() - 1; Index >= 0; --Index)
+		{
+			if (Groups[Index].RepresentativeTime < InRange.GetUpperBoundValue())
+			{
+				// Just entered the range
+				return Groups[Index].RepresentativeTime;
+			}
+			else if (InRange.HasLowerBound() && Groups[Index].RepresentativeTime < InRange.GetLowerBoundValue())
+			{
+				// No longer inside the range
+				return TOptional<float>();
+			}
+		}
+	}
+	else
+	{
+		for (int32 Index = 0; Index < Groups.Num(); ++Index)
+		{
+			if (Groups[Index].RepresentativeTime > InRange.GetLowerBoundValue())
+			{
+				// Just entered the range
+				return Groups[Index].RepresentativeTime;
+			}
+			else if (InRange.HasUpperBound() && Groups[Index].RepresentativeTime > InRange.GetUpperBoundValue())
+			{
+				// No longer inside the range
+				return TOptional<float>();
+			}
+		}
+	}
+
+	return TOptional<float>();
+}
+
+const FKeyGrouping* FGroupedKeyCollection::FindGroup(FKeyHandle InHandle) const
 {
 	if (auto* IndexEntry = GlobalIndex.Find(IndexKey))
 	{
@@ -143,7 +278,7 @@ const FKeyGrouping* FGroupedKeyArea::FindGroup(FKeyHandle InHandle) const
 	return nullptr;
 }
 
-FKeyGrouping* FGroupedKeyArea::FindGroup(FKeyHandle InHandle)
+FKeyGrouping* FGroupedKeyCollection::FindGroup(FKeyHandle InHandle)
 {
 	if (auto* IndexEntry = GlobalIndex.Find(IndexKey))
 	{
@@ -155,13 +290,13 @@ FKeyGrouping* FGroupedKeyArea::FindGroup(FKeyHandle InHandle)
 	return nullptr;
 }
 
-FLinearColor FGroupedKeyArea::GetKeyTint(FKeyHandle InHandle) const
+FLinearColor FGroupedKeyCollection::GetKeyTint(FKeyHandle InHandle) const
 {
 	// Everything is untinted for now
 	return FLinearColor::White;
 }
 
-const FSlateBrush* FGroupedKeyArea::GetBrush(FKeyHandle InHandle) const
+const FSlateBrush* FGroupedKeyCollection::GetBrush(FKeyHandle InHandle) const
 {
 	static const FSlateBrush* PartialKeyBrush = FEditorStyle::GetBrush("Sequencer.PartialKey");
 	const auto* Group = FindGroup(InHandle);
@@ -178,6 +313,27 @@ const FSlateBrush* FGroupedKeyArea::GetBrush(FKeyHandle InHandle) const
 	return nullptr;
 }
 
+FGroupedKeyArea::FGroupedKeyArea(FSequencerDisplayNode& InNode, int32 InSectionIndex)
+	: Section(nullptr)
+{
+	// Calling this virtual function is safe here, as it's just called through standard name-lookup, not runtime dispatch
+	FGroupedKeyCollection::InitializeRecursive(InNode, InSectionIndex);
+
+	for (const TSharedPtr<IKeyArea>& KeyArea : KeyAreas)
+	{
+		auto* OwningSection = KeyArea->GetOwningSection();
+		// Ensure they all belong to the same section
+		ensure(!Section || Section == OwningSection);
+		Section = OwningSection;
+	}
+
+	if (Section)
+	{
+		IndexKey = FIndexKey(InNode.GetPathName(), Section);
+		UpdateIndex();
+	}
+}
+
 TArray<FKeyHandle> FGroupedKeyArea::GetUnsortedKeyHandles() const
 {
 	TArray<FKeyHandle> Array;
diff --git a/Engine/Source/Editor/Sequencer/Private/GroupedKeyArea.h b/Engine/Source/Editor/Sequencer/Private/GroupedKeyArea.h
--- a/Engine/Source/Editor/Sequencer/Private/GroupedKeyArea.h
+++ b/Engine/Source/Editor/Sequencer/Private/GroupedKeyArea.h
@@ -3,6 +3,7 @@
 #pragma once
 
 #include "IKeyArea.h"
+#include "ISequencerKeyCollection.h"
 
 
 class FSequencerDisplayNode;
@@ -76,13 +77,33 @@ struct FKeyGrouping
 
 
 /** A key area that represents multiple, distinct key areas */
-class FGroupedKeyArea
-	: public IKeyArea
+class FGroupedKeyCollection
+	: public ISequencerKeyCollection
 {
 public:
 
-	/** Construct the area from a display node */
-	FGroupedKeyArea(FSequencerDisplayNode& InNode, int32 InSectionIndex);
+	/** Constructor */
+	FGroupedKeyCollection();
+
+public:
+
+	/** Initialize this key collection from the specified nodes. Only gathers keys from those explicitly specified. */
+	virtual void InitializeExplicit(const TArray<FSequencerDisplayNode*>& InNodes, float DuplicateThreshold = SMALL_NUMBER) override;
+
+	/** Initialize this key collection from the specified nodes. Gathers keys from all child nodes. */
+	virtual void InitializeRecursive(const TArray<FSequencerDisplayNode*>& InNodes, float DuplicateThreshold = SMALL_NUMBER) override;
+	
+	/** Initialize this key collection from the specified node and section index. */
+	virtual void InitializeRecursive(FSequencerDisplayNode& InNode, int32 InSectionIndex, float DuplicateThreshold = SMALL_NUMBER) override;
+
+	/** Iterate the keys contained within this collection */
+	virtual void IterateKeys(const TFunctionRef<bool(float)>& Iter) override;
+
+	/** Get a value specifying how close keys need to be in order to be considered equal by this collection */
+	virtual float GetKeyGroupingThreshold() const override;
+
+	/** Find the first key in the given range. */
+	virtual TOptional<float> FindFirstKeyInRange(const TRange<float>& InRange, EFindKeyDirection Direction) const override;
 
 public:
 
@@ -92,6 +113,47 @@ public:
 	/** Get the brush name for the specified key handle */
 	const FSlateBrush* GetBrush(FKeyHandle InHandle) const;
 
+protected:
+
+	/** Add the specified key area to this collection */
+	void AddKeyArea(const TSharedRef<IKeyArea>& InKeyArea);
+
+	/** Remove key groups that reside at similar times */
+	void RemoveDuplicateKeys(float DuplicateThreshold);
+
+	/** Update the persistent index with our new key handles */
+	void UpdateIndex() const;
+
+	/** Helper function to find the group of keys relating to the specified handle */
+	const FKeyGrouping* FindGroup(FKeyHandle InHandle) const;
+
+	/** Helper function to find the group of keys relating to the specified handle */
+	FKeyGrouping* FindGroup(FKeyHandle InHandle);
+
+protected:
+	/** Array of (child) key areas that we are reflecting */
+	TArray<TSharedRef<IKeyArea>> KeyAreas;
+
+	/** Generated array of groups of keys harvested from the above array */
+	TArray<FKeyGrouping> Groups;
+
+	/** Key into our index */
+	FIndexKey IndexKey;
+
+	/** Value specifying how close keys need to be in order to be considered equal by this collection */
+	float GroupingThreshold;
+};
+
+/** A key area that represents multiple, distinct key areas */
+class FGroupedKeyArea
+	: public IKeyArea
+	, public FGroupedKeyCollection
+{
+public:
+
+	/** Construct the area from a display node */
+	FGroupedKeyArea(FSequencerDisplayNode& InNode, int32 InSectionIndex);
+
 public:
 
 	// IKeyArea interface
@@ -121,28 +183,8 @@ public:
 	virtual void PasteKeys(const FMovieSceneClipboardKeyTrack& KeyTrack, const FMovieSceneClipboardEnvironment& SrcEnvironment, const FSequencerPasteEnvironment& DstEnvironment) override;
 	virtual TOptional<FLinearColor> GetColor() override;
 
-protected:
-
-	/** Update the persistent index with our new key handles */
-	void UpdateIndex() const;
-
-	/** Helper function to find the group of keys relating to the specified handle */
-	const FKeyGrouping* FindGroup(FKeyHandle InHandle) const;
-
-	/** Helper function to find the group of keys relating to the specified handle */
-	FKeyGrouping* FindGroup(FKeyHandle InHandle);
-
 protected:
 	
 	/** Pointer to the section to which this key area relates */
 	UMovieSceneSection* Section;
-
-	/** Array of (child) key areas that we are reflecting */
-	TArray<TSharedRef<IKeyArea>> KeyAreas;
-
-	/** Generated array of groups of keys harvested from the above array */
-	TArray<FKeyGrouping> Groups;
-
-	/** Key into our index */
-	FIndexKey IndexKey;
 };
diff --git a/Engine/Source/Editor/Sequencer/Private/IntegralKeyDetailsCustomization.cpp b/Engine/Source/Editor/Sequencer/Private/IntegralKeyDetailsCustomization.cpp
new file mode 100644
--- /dev/null
+++ b/Engine/Source/Editor/Sequencer/Private/IntegralKeyDetailsCustomization.cpp
@@ -0,0 +1,87 @@
+// Copyright 1998-2016 Epic Games, Inc. All Rights Reserved.
+
+#include "SequencerPrivatePCH.h"
+#include "IntegralKeyDetailsCustomization.h"
+#include "PropertyEditing.h"
+#include "DetailLayoutBuilder.h"
+#include "MovieSceneBoolSection.h"
+#include "MovieSceneParticleSection.h"
+#include "MovieSceneToolHelpers.h"
+
+#define LOCTEXT_NAMESPACE "IntegralKeyDetailsCustomization"
+
+TSharedRef<class IDetailCustomization> FIntegralKeyDetailsCustomization::MakeInstance(TWeakObjectPtr<const UMovieSceneSection> InSection)
+{
+	return MakeShareable(new FIntegralKeyDetailsCustomization(InSection));
+}
+
+BEGIN_SLATE_FUNCTION_BUILD_OPTIMIZATION
+void FIntegralKeyDetailsCustomization::CustomizeDetails(IDetailLayoutBuilder& DetailBuilder)
+{
+	if(Section.IsValid())
+	{
+		TSharedPtr<IPropertyHandle> ValueProperty = DetailBuilder.GetProperty(GET_MEMBER_NAME_CHECKED(FIntegralKey, Value));
+		if(ValueProperty.IsValid())
+		{
+			TSharedPtr<SWidget> ValueWidget = nullptr;
+
+			if(Section->IsA<UMovieSceneBoolSection>())
+			{
+				ValueWidget = SNew(SCheckBox)
+					.IsChecked_Lambda([ValueProperty]()
+					{
+						int32 Value = 0;
+						if(ValueProperty->GetValue(Value) == FPropertyAccess::Success)
+						{
+							return (Value != 0) ? ECheckBoxState::Checked : ECheckBoxState::Unchecked;
+						}
+
+						return ECheckBoxState::Undetermined;
+					})
+					.OnCheckStateChanged_Lambda([ValueProperty](ECheckBoxState InState) 
+					{ 
+						int32 Value = (InState == ECheckBoxState::Checked) ? 1 : 0;
+						ValueProperty->SetValue(Value);
+					});
+			}
+			else if(Section->IsA<UMovieSceneParticleSection>())
+			{
+				UEnum* Enum = FindObject<UEnum>(ANY_PACKAGE, TEXT("EParticleKey"), true);
+				check(Enum);
+
+				ValueWidget = MovieSceneToolHelpers::MakeEnumComboBox(
+					Enum, 
+					TAttribute<int32>::Create(TAttribute<int32>::FGetter::CreateLambda([ValueProperty]() 
+					{ 
+						int32 Value = 0;
+						ValueProperty->GetValue(Value);
+						return Value;
+					})),
+					FOnEnumSelectionChanged::CreateLambda([ValueProperty](int32 Selection, ESelectInfo::Type SelectionType)
+					{
+						ValueProperty->SetValue(Selection);
+					}),
+					TAttribute<TOptional<uint8>>());
+			}
+			else
+			{
+				ValueWidget = ValueProperty->CreatePropertyValueWidget();
+			}
+
+			DetailBuilder.EditCategory("Key")
+				.AddProperty("Value")
+				.CustomWidget()
+				.NameContent()
+				[
+					ValueProperty->CreatePropertyNameWidget()
+				]
+				.ValueContent()
+				[
+					ValueWidget.ToSharedRef()
+				];
+		}
+	}
+}
+END_SLATE_FUNCTION_BUILD_OPTIMIZATION
+
+#undef LOCTEXT_NAMESPACE
\ No newline at end of file
