Copying //UE4/Dev-Sequencer to //UE4/Main (Source: //UE4/Dev-Sequencer @ 2945541)

diff --git a/Engine/Source/Editor/MovieSceneTools/Public/MatineeImportTools.h b/Engine/Source/Editor/MovieSceneTools/Public/MatineeImportTools.h
new file mode 100644
--- /dev/null
+++ b/Engine/Source/Editor/MovieSceneTools/Public/MatineeImportTools.h
@@ -0,0 +1,81 @@
+// Copyright 1998-2016 Epic Games, Inc. All Rights Reserved.
+#pragma once
+
+#include "Matinee/InterpTrackToggle.h"
+#include "MovieSceneParticleSection.h"
+
+class UInterpTrackColorProp;
+class UInterpTrackBoolProp;
+class UInterpTrackFloatBase;
+class UInterpTrackMove;
+class UInterpTrackToggle;
+class UInterpTrackAnimControl;
+class UInterpTrackSound;
+class UInterpTrackFade;
+class UInterpTrackDirector;
+class UInterpTrackEvent;
+class UInterpTrackVisibility;
+
+class UMovieSceneColorTrack;
+class UMovieSceneBoolTrack;
+class UMovieSceneFloatTrack;
+class UMovieScene3DTransformTrack;
+class UMovieSceneParticleTrack;
+class UMovieSceneSkeletalAnimationTrack;
+class UMovieSceneAudioTrack;
+class UMovieSceneFadeTrack;
+class UMovieSceneCameraCutTrack;
+class UMovieSceneEventTrack;
+class UMovieSceneVisibilityTrack;
+
+class UMovieSceneSequence;
+
+class MOVIESCENETOOLS_API FMatineeImportTools
+{
+public:
+
+	/** Converts a matinee interpolation mode to its equivalent rich curve interpolation mode. */
+	static ERichCurveInterpMode MatineeInterpolationToRichCurveInterpolation( EInterpCurveMode CurveMode );
+
+	/** Converts a matinee interpolation mode to its equivalent rich curve tangent mode. */
+	static ERichCurveTangentMode MatineeInterpolationToRichCurveTangent( EInterpCurveMode CurveMode );
+
+	/** Tries to convert a matinee toggle to a particle key. */
+	static bool TryConvertMatineeToggleToOutParticleKey( ETrackToggleAction ToggleAction, EParticleKey::Type& OutParticleKey );
+
+	/** Adds a key to a rich curve based on matinee curve key data. */
+	static void SetOrAddKey( FRichCurve& Curve, float Time, float Value, float ArriveTangent, float LeaveTangent, EInterpCurveMode MatineeInterpMode );
+
+	/** Copies keys from a matinee bool track to a sequencer bool track. */
+	static bool CopyInterpBoolTrack( UInterpTrackBoolProp* MatineeBoolTrack, UMovieSceneBoolTrack* BoolTrack );
+
+	/** Copies keys from a matinee float track to a sequencer float track. */
+	static bool CopyInterpFloatTrack( UInterpTrackFloatBase* MatineeFloatTrack, UMovieSceneFloatTrack* FloatTrack );
+
+	/** Copies keys from a matinee move track to a sequencer transform track. */
+	static bool CopyInterpMoveTrack( UInterpTrackMove* MoveTrack, UMovieScene3DTransformTrack* TransformTrack );
+
+	/** Copies keys from a matinee color track to a sequencer color track. */
+	static bool CopyInterpColorTrack( UInterpTrackColorProp* ColorPropTrack, UMovieSceneColorTrack* ColorTrack );
+
+	/** Copies keys from a matinee toggle track to a sequencer particle track. */
+	static bool CopyInterpParticleTrack( UInterpTrackToggle* MatineeToggleTrack, UMovieSceneParticleTrack* ParticleTrack );
+
+	/** Copies keys from a matinee anim control track to a sequencer skeletal animation track. */
+	static bool CopyInterpAnimControlTrack( UInterpTrackAnimControl* MatineeAnimControlTrack, UMovieSceneSkeletalAnimationTrack* SkeletalAnimationTrack, float EndPlaybackRange );
+
+	/** Copies keys from a matinee sound track to a sequencer audio track. */
+	static bool CopyInterpSoundTrack( UInterpTrackSound* MatineeSoundTrack, UMovieSceneAudioTrack* AudioTrack );
+
+	/** Copies keys from a matinee fade track to a sequencer fade track. */
+	static bool CopyInterpFadeTrack( UInterpTrackFade* MatineeFadeTrack, UMovieSceneFadeTrack* FadeTrack );
+
+	/** Copies keys from a matinee director track to a sequencer camera cut track. */
+	static bool CopyInterpDirectorTrack( UInterpTrackDirector* DirectorTrack, UMovieSceneCameraCutTrack* CameraCutTrack, AMatineeActor* MatineeActor, UMovieSceneSequence* MovieSceneSequence );
+
+	/** Copies keys from a matinee event track to a sequencer event track. */
+	static bool CopyInterpEventTrack( UInterpTrackEvent* MatineeEventTrack, UMovieSceneEventTrack* EventTrack );
+
+	/** Copies keys from a matinee visibility track to a sequencer visibility track. */
+	static bool CopyInterpVisibilityTrack( UInterpTrackVisibility* MatineeVisibilityTrack, UMovieSceneVisibilityTrack* VisibilityTrack );
+};
\ No newline at end of file
diff --git a/Engine/Source/Editor/PlacementMode/Private/SPlacementModeTools.cpp b/Engine/Source/Editor/PlacementMode/Private/SPlacementModeTools.cpp
--- a/Engine/Source/Editor/PlacementMode/Private/SPlacementModeTools.cpp
+++ b/Engine/Source/Editor/PlacementMode/Private/SPlacementModeTools.cpp
@@ -369,21 +369,23 @@ void SPlacementModeTools::Construct( const FArguments& InArgs )
 
 					+ SOverlay::Slot()
 					[
-						SNew(SHorizontalBox)
-
-						+ SHorizontalBox::Slot()
+						SAssignNew(DataDrivenContent, SBox)
 						[
-							SAssignNew(ListView, SListView<TSharedPtr<FPlaceableItem>>)
-							.ListItemsSource(&FilteredItems)
-							.OnGenerateRow(this, &SPlacementModeTools::OnGenerateWidgetForItem)
-							.Visibility(this, &SPlacementModeTools::GetListViewVisibility)
-							.ExternalScrollbar(ScrollBar)
-						]
-
-						+ SHorizontalBox::Slot()
-						.AutoWidth()
-						[
-							ScrollBar
+							SNew(SHorizontalBox)
+
+							+ SHorizontalBox::Slot()
+							[
+								SAssignNew(ListView, SListView<TSharedPtr<FPlaceableItem>>)
+								.ListItemsSource(&FilteredItems)
+								.OnGenerateRow(this, &SPlacementModeTools::OnGenerateWidgetForItem)
+								.ExternalScrollbar(ScrollBar)
+							]
+
+							+ SHorizontalBox::Slot()
+							.AutoWidth()
+							[
+								ScrollBar
+							]
 						]
 					]
 				]
@@ -453,7 +455,7 @@ void SPlacementModeTools::UpdateFilteredItems()
 		CustomContent->SetContent(Category->CustomGenerator());
 
 		CustomContent->SetVisibility(EVisibility::Visible);
-		ListView->SetVisibility(EVisibility::Collapsed);
+		DataDrivenContent->SetVisibility(EVisibility::Collapsed);
 	}
 	else
 	{
@@ -486,7 +488,7 @@ void SPlacementModeTools::UpdateFilteredItems()
 		}
 
 		CustomContent->SetVisibility(EVisibility::Collapsed);
-		ListView->SetVisibility(EVisibility::Visible);
+		DataDrivenContent->SetVisibility(EVisibility::Visible);
 		ListView->RequestListRefresh();
 	}
 }
@@ -510,15 +512,6 @@ EVisibility SPlacementModeTools::GetFailedSearchVisibility() const
 	return EVisibility::Visible;
 }
 
-EVisibility SPlacementModeTools::GetListViewVisibility()const
-{
-	if (IsSearchActive() || FilteredItems.Num())
-	{
-		return EVisibility::Visible;
-	}
-	return EVisibility::Collapsed;
-}
-
 EVisibility SPlacementModeTools::GetTabsVisibility() const
 {
 	return IsSearchActive() ? EVisibility::Collapsed : EVisibility::Visible;
diff --git a/Engine/Source/Editor/PlacementMode/Private/SPlacementModeTools.h b/Engine/Source/Editor/PlacementMode/Private/SPlacementModeTools.h
--- a/Engine/Source/Editor/PlacementMode/Private/SPlacementModeTools.h
+++ b/Engine/Source/Editor/PlacementMode/Private/SPlacementModeTools.h
@@ -124,6 +124,9 @@ private:
 
 	/** Custom content slot, where a category has a custom generator */
 	TSharedPtr<SBox> CustomContent;
+	
+	/** Content container for any data driven content */
+	TSharedPtr<SBox> DataDrivenContent;
 
 	/* The search box used to update the filter text */
 	TSharedPtr<SSearchBox> SearchBoxPtr;
diff --git a/Engine/Source/Editor/PropertyEditor/Private/PropertyCustomizationHelpers.cpp b/Engine/Source/Editor/PropertyEditor/Private/PropertyCustomizationHelpers.cpp
--- a/Engine/Source/Editor/PropertyEditor/Private/PropertyCustomizationHelpers.cpp
+++ b/Engine/Source/Editor/PropertyEditor/Private/PropertyCustomizationHelpers.cpp
@@ -15,6 +15,7 @@
 #include "SPropertyEditorAsset.h"
 #include "SPropertyEditorClass.h"
 #include "SPropertyEditorInteractiveActorPicker.h"
+#include "SPropertyEditorSceneDepthPicker.h"
 #include "SHyperlink.h"
 #include "SWidgetSwitcher.h"
 
@@ -248,6 +249,14 @@ namespace PropertyCustomizationHelpers
 			.OnActorSelected( OnActorSelectedFromPicker );
 	}
 
+	TSharedRef<SWidget> MakeSceneDepthPicker(FOnSceneDepthLocationSelected OnSceneDepthLocationSelected)
+	{
+		return
+			SNew(SPropertyEditorSceneDepthPicker)
+			.ToolTipText(LOCTEXT("PickSceneDepthLabel", "Sample Scene Depth from scene"))
+			.OnSceneDepthLocationSelected(OnSceneDepthLocationSelected);
+	}
+
 	TSharedRef<SWidget> MakeEditConfigHierarchyButton(FSimpleDelegate OnEditConfigClicked, TAttribute<FText> OptionalToolTipText, TAttribute<bool> IsEnabled)
 	{
 		return
diff --git a/Engine/Source/Editor/PropertyEditor/Public/PropertyCustomizationHelpers.h b/Engine/Source/Editor/PropertyEditor/Public/PropertyCustomizationHelpers.h
--- a/Engine/Source/Editor/PropertyEditor/Public/PropertyCustomizationHelpers.h
+++ b/Engine/Source/Editor/PropertyEditor/Public/PropertyCustomizationHelpers.h
@@ -8,6 +8,7 @@
 #include "FilterCollection.h"
 #include "SResetToDefaultMenu.h"
 #include "ActorPickerMode.h"
+#include "SceneDepthPickerMode.h"
 #include "PropertyHandle.h"
 #include "DetailWidgetRow.h"
 
@@ -34,6 +35,7 @@ namespace PropertyCustomizationHelpers
 	PROPERTYEDITOR_API TSharedRef<SWidget> MakeActorPickerAnchorButton( FOnGetActorFilters OnGetActorFilters, FOnActorSelected OnActorSelectedFromPicker );
 	PROPERTYEDITOR_API TSharedRef<SWidget> MakeActorPickerWithMenu( AActor* const InitialActor, const bool AllowClear, FOnShouldFilterActor ActorFilter, FOnActorSelected OnSet, FSimpleDelegate OnClose, FSimpleDelegate OnUseSelected );
 	PROPERTYEDITOR_API TSharedRef<SWidget> MakeInteractiveActorPicker(FOnGetAllowedClasses OnGetAllowedClasses, FOnShouldFilterActor OnShouldFilterActor, FOnActorSelected OnActorSelectedFromPicker);
+	PROPERTYEDITOR_API TSharedRef<SWidget> MakeSceneDepthPicker(FOnSceneDepthLocationSelected OnSceneDepthLocationSelected);
 	PROPERTYEDITOR_API TSharedRef<SWidget> MakeEditConfigHierarchyButton(FSimpleDelegate OnEditConfigClicked, TAttribute<FText> OptionalToolTipText = FText(), TAttribute<bool> IsEnabled = true);
 
 	/** @return the UBoolProperty edit condition property if one exists. */
diff --git a/Engine/Source/Editor/SceneDepthPickerMode/Private/SceneDepthPickerModePrivatePCH.h b/Engine/Source/Editor/SceneDepthPickerMode/Private/SceneDepthPickerModePrivatePCH.h
new file mode 100644
--- /dev/null
+++ b/Engine/Source/Editor/SceneDepthPickerMode/Private/SceneDepthPickerModePrivatePCH.h
@@ -0,0 +1,8 @@
+// Copyright 1998-2015 Epic Games, Inc. All Rights Reserved.
+
+#pragma once
+
+#include "UnrealEd.h"
+
+#include "SceneDepthPickerMode.h"
+#include "EditorModeSceneDepthPicker.h"
diff --git a/Engine/Source/Editor/SceneDepthPickerMode/Public/SceneDepthPickerMode.h b/Engine/Source/Editor/SceneDepthPickerMode/Public/SceneDepthPickerMode.h
new file mode 100644
--- /dev/null
+++ b/Engine/Source/Editor/SceneDepthPickerMode/Public/SceneDepthPickerMode.h
@@ -0,0 +1,29 @@
+// Copyright 1998-2015 Epic Games, Inc. All Rights Reserved.
+
+#pragma once
+
+DECLARE_DELEGATE_OneParam( FOnSceneDepthLocationSelected, FVector );
+
+/**
+ * Scene depth picker mode module
+ */
+class SCENEDEPTHPICKERMODE_API FSceneDepthPickerModeModule : public IModuleInterface
+{
+public:
+	// IModuleInterface
+	virtual void StartupModule() override;
+	virtual void ShutdownModule() override;
+	// End of IModuleInterface
+
+	/** 
+	 * Enter scene depth picking mode (note: will cancel any current scene depth picking)
+	 * @param	InOnDepthSelected		Delegate to call when user has selected a location to sample the scene depth
+	 */
+	void BeginSceneDepthPickingMode(FOnSceneDepthLocationSelected InOnSceneDepthLocationSelected);
+
+	/** Exit scene depth picking mode */
+	void EndSceneDepthPickingMode();
+
+	/** @return Whether or not scene depth picking mode is currently active */
+	bool IsInSceneDepthPickingMode() const;
+};
diff --git a/Engine/Source/Editor/SceneDepthPickerMode/SceneDepthPickerMode.Build.cs b/Engine/Source/Editor/SceneDepthPickerMode/SceneDepthPickerMode.Build.cs
new file mode 100644
--- /dev/null
+++ b/Engine/Source/Editor/SceneDepthPickerMode/SceneDepthPickerMode.Build.cs
@@ -0,0 +1,22 @@
+// Copyright 1998-2015 Epic Games, Inc. All Rights Reserved.
+
+using UnrealBuildTool;
+
+public class SceneDepthPickerMode : ModuleRules
+{
+    public SceneDepthPickerMode(TargetInfo Target)
+	{
+		PrivateDependencyModuleNames.AddRange(
+			new string[] {
+				"Core",
+				"CoreUObject",
+				"Engine",
+                "InputCore",
+                "Slate",
+                "SlateCore",
+                "EditorStyle",
+				"UnrealEd",
+			}
+		);
+	}
+}
diff --git a/Engine/Source/Editor/SequenceRecorder/Private/ActorRecording.cpp b/Engine/Source/Editor/SequenceRecorder/Private/ActorRecording.cpp
--- a/Engine/Source/Editor/SequenceRecorder/Private/ActorRecording.cpp
+++ b/Engine/Source/Editor/SequenceRecorder/Private/ActorRecording.cpp
@@ -4,13 +4,8 @@
 #include "ActorRecording.h"
 #include "SequenceRecorderSettings.h"
 #include "SequenceRecorderUtils.h"
-#include "IMovieScenePropertyRecorder.h"
-#include "MovieSceneSpawnPropertyRecorder.h"
-#include "MovieScene3DTransformPropertyRecorder.h"
-#include "MovieSceneVisibilityPropertyRecorder.h"
-#include "MovieSceneAnimationPropertyRecorder.h"
-#include "MovieSceneParticleTrackPropertyRecorder.h"
-#include "MovieScene3DAttachPropertyRecorder.h"
+#include "MovieScene3DTransformSectionRecorder.h"
+#include "MovieSceneAnimationSectionRecorder.h"
 #include "AssetSelection.h"
 #include "MovieSceneSkeletalAnimationTrack.h"
 #include "MovieSceneSkeletalAnimationSection.h"
@@ -24,8 +19,11 @@
 #include "CameraRig_Crane.h"
 #include "CameraRig_Rail.h"
 #include "Animation/SkeletalMeshActor.h"
+#include "SequenceRecorder.h"
+#include "Runtime/Core/Public/Features/IModularFeatures.h"
 
 static const FName SequencerActorTag(TEXT("SequencerActor"));
+static const FName MovieSceneSectionRecorderFactoryName("MovieSceneSectionRecorderFactory");
 
 UActorRecording::UActorRecording(const FObjectInitializer& ObjectInitializer)
 	: Super(ObjectInitializer)
@@ -52,15 +50,16 @@ bool UActorRecording::IsRelevantForRecording(AActor* Actor)
 	}
 
 	TInlineComponentArray<USceneComponent*> SceneComponents(Actor);
+	const USequenceRecorderSettings* Settings = GetDefault<USequenceRecorderSettings>();
 
 	for(USceneComponent* SceneComponent : SceneComponents)
 	{
-		if(SceneComponent->IsA<USkeletalMeshComponent>() ||
-			SceneComponent->IsA<UStaticMeshComponent>() ||
-			SceneComponent->IsA<UParticleSystemComponent>() || 
-			SceneComponent->IsA<ULightComponent>())
+		for (TSubclassOf<USceneComponent> ComponentClass : Settings->ComponentClassesToRecord)
 		{
-			return true;
+			if (SceneComponent->IsA(ComponentClass))
+			{
+				return true;
+			}
 		}
 	}
 
@@ -79,9 +78,10 @@ bool UActorRecording::StartRecording(ULevelSequence* CurrentSequence, float Curr
 		}
 		else
 		{
-			TSharedPtr<FMovieSceneAnimationPropertyRecorder> AnimationRecorder = MakeShareable(new FMovieSceneAnimationPropertyRecorder(AnimationSettings, TargetAnimation.Get()));
-			AnimationRecorder->CreateSection(ActorToRecord.Get(), nullptr, FGuid(), 0.0f, true);
-			PropertyRecorders.Add(AnimationRecorder);			
+			TSharedPtr<FMovieSceneAnimationSectionRecorder> AnimationRecorder = MakeShareable(new FMovieSceneAnimationSectionRecorder(AnimationSettings, TargetAnimation.Get()));
+			AnimationRecorder->CreateSection(ActorToRecord.Get(), nullptr, FGuid(), 0.0f);
+			AnimationRecorder->Record(0.0f);
+			SectionRecorders.Add(AnimationRecorder);			
 		}
 	}
 
@@ -163,6 +163,31 @@ void UActorRecording::SyncTrackedComponents(bool bIncludeNonCDO/*=true*/)
 	TArray<USceneComponent*> NewComponentArray;
 	GetSceneComponents(NewComponentArray, bIncludeNonCDO);
 
+	// Expire section recorders that are watching components no longer attached to our actor
+	TSet<USceneComponent*> ExpiredComponents;
+	for (TWeakObjectPtr<USceneComponent>& WeakComponent : TrackedComponents)
+	{
+		if (USceneComponent* Component = WeakComponent.Get())
+		{
+			ExpiredComponents.Add(Component);
+		}
+	}
+	for (USceneComponent* Component : NewComponentArray)
+	{
+		ExpiredComponents.Remove(Component);
+	}
+
+	for (TSharedPtr<IMovieSceneSectionRecorder>& SectionRecorder : SectionRecorders)
+	{
+		if (USceneComponent* Component = Cast<USceneComponent>(SectionRecorder->GetSourceObject()))
+		{
+			if (ExpiredComponents.Contains(Component))
+			{
+				SectionRecorder->InvalidateObjectToRecord();
+			}
+		}
+	}
+
 	TrackedComponents.Reset(NewComponentArray.Num());
 	for(USceneComponent* SceneComponent : NewComponentArray)
 	{
@@ -173,9 +198,9 @@ void UActorRecording::SyncTrackedComponents(bool bIncludeNonCDO/*=true*/)
 void UActorRecording::InvalidateObjectToRecord()
 {
 	ActorToRecord = nullptr;
-	for(auto& PropertyRecorder : PropertyRecorders)
+	for(auto& SectionRecorder : SectionRecorders)
 	{
-		PropertyRecorder->InvalidateObjectToRecord();
+		SectionRecorder->InvalidateObjectToRecord();
 	}
 }
 
@@ -183,12 +208,14 @@ bool UActorRecording::ValidComponent(USceneComponent* SceneComponent) const
 {
 	if(SceneComponent != nullptr)
 	{
-		const bool bIsValidComponent = 
-			SceneComponent->IsA<USkeletalMeshComponent>() ||
-			SceneComponent->IsA<UStaticMeshComponent>() ||
-			SceneComponent->IsA<UParticleSystemComponent>();
-
-		return bIsValidComponent;
+		const USequenceRecorderSettings* Settings = GetDefault<USequenceRecorderSettings>();
+		for (TSubclassOf<USceneComponent> ComponentClass : Settings->ComponentClassesToRecord)
+		{
+			if (SceneComponent->IsA(ComponentClass))
+			{
+				return true;
+			}
+		}
 	}
 
 	return false;
@@ -241,44 +268,24 @@ void UActorRecording::StartRecordingActorProperties(ULevelSequence* CurrentSeque
 		UMovieScene* MovieScene = CurrentSequence->GetMovieScene();
 
 		AActor* Actor = ActorToRecord.Get();
-		FString BlueprintName = GetUniqueSpawnableName(MovieScene, Actor->GetName());
+		FString TemplateName = GetUniqueSpawnableName(MovieScene, Actor->GetName());
 
 		UClass* ActorClass = Actor->GetClass();
-			
-		UBlueprint* Blueprint = nullptr;
+		AActor* ObjectTemplate = DuplicateObject<AActor>(Actor, MovieScene, *TemplateName);
 
-		if(ActorClass->IsChildOf<AWorldSettings>())
+		if (ObjectTemplate)
 		{
-			// we need to create our own substitute blueprint for world settings actors as they are not blueprintable by default
-			Blueprint = FKismetEditorUtilities::CreateBlueprint(AActor::StaticClass(), MovieScene, *BlueprintName, EBlueprintType::BPTYPE_Normal, UBlueprint::StaticClass(), UBlueprintGeneratedClass::StaticClass(), NAME_Name);
-		}
-		else
-		{
-			UBlueprint* NewBlueprint = FKismetEditorUtilities::CreateBlueprint(ActorClass, MovieScene, *BlueprintName, EBlueprintType::BPTYPE_Normal, UBlueprint::StaticClass(), UBlueprintGeneratedClass::StaticClass(), NAME_Name);
-			if(NewBlueprint)
+			TInlineComponentArray<USkeletalMeshComponent*> SkeletalMeshComponents;
+			ObjectTemplate->GetComponents(SkeletalMeshComponents);
+			for(USkeletalMeshComponent* SkeletalMeshComponent : SkeletalMeshComponents)
 			{
-				// if the duplicated BP has any skel mesh components, set them to single anim instance
-				AActor* CDO = CastChecked<AActor>(NewBlueprint->GeneratedClass->ClassDefaultObject);
-				if(CDO)
-				{
-					TInlineComponentArray<USkeletalMeshComponent*> SkeletalMeshComponents;
-					CDO->GetComponents(SkeletalMeshComponents);	
-					for(USkeletalMeshComponent* SkeletalMeshComponent : SkeletalMeshComponents)
-					{
-						SkeletalMeshComponent->SetAnimationMode(EAnimationMode::AnimationSingleNode);
-						SkeletalMeshComponent->bEnableUpdateRateOptimizations = false;
-						SkeletalMeshComponent->MeshComponentUpdateFlag = EMeshComponentUpdateFlag::AlwaysTickPoseAndRefreshBones;
-						SkeletalMeshComponent->ForcedLodModel = 1;
-					}
-				}
-
-				Blueprint = NewBlueprint;
+				SkeletalMeshComponent->SetAnimationMode(EAnimationMode::AnimationSingleNode);
+				SkeletalMeshComponent->bEnableUpdateRateOptimizations = false;
+				SkeletalMeshComponent->MeshComponentUpdateFlag = EMeshComponentUpdateFlag::AlwaysTickPoseAndRefreshBones;
+				SkeletalMeshComponent->ForcedLodModel = 1;
 			}
-		}
 
-		if(Blueprint)
-		{
-			Guid = MovieScene->AddSpawnable(BlueprintName, Blueprint);
+			Guid = MovieScene->AddSpawnable(TemplateName, *ObjectTemplate);
 		}
 
 		// now add tracks to record
@@ -288,7 +295,9 @@ void UActorRecording::StartRecordingActorProperties(ULevelSequence* CurrentSeque
 			FindOrAddFolder(MovieScene);
 
 			// force set recording to record translations as we need this with no animation
-			ActorSettings.bRecordTransforms = true;
+			UMovieScene3DTransformSectionRecorderSettings* TransformSettings = ActorSettings.GetSettingsObject<UMovieScene3DTransformSectionRecorderSettings>();
+			check(TransformSettings);
+			TransformSettings->bRecordTransforms = true;
 
 			// grab components so we can track attachments
 			// don't include non-CDO here as they wont be part of our initial BP (duplicated above)
@@ -316,67 +325,55 @@ void UActorRecording::StartRecordingActorProperties(ULevelSequence* CurrentSeque
 				}
 			}
 
-			TSharedPtr<FMovieSceneAnimationPropertyRecorder> FirstAnimRecorder = nullptr;
+			TSharedPtr<FMovieSceneAnimationSectionRecorder> FirstAnimRecorder = nullptr;
 			for(USceneComponent* SceneComponent : ValidSceneComponents)
 			{
-				TSharedPtr<FMovieSceneAnimationPropertyRecorder> AnimRecorder = StartRecordingComponentProperties(SceneComponent->GetFName(), SceneComponent, ActorToRecord.Get(), CurrentSequence, CurrentSequenceTime);
+				TSharedPtr<FMovieSceneAnimationSectionRecorder> AnimRecorder = StartRecordingComponentProperties(SceneComponent->GetFName(), SceneComponent, ActorToRecord.Get(), CurrentSequence, CurrentSequenceTime);
 				if(!FirstAnimRecorder.IsValid() && AnimRecorder.IsValid())
 				{
 					FirstAnimRecorder = AnimRecorder;
 				}
 			}
 
-			// add a spawn track so we actually get spawned
-			TSharedPtr<FMovieSceneSpawnPropertyRecorder> SpawnRecorder = MakeShareable(new FMovieSceneSpawnPropertyRecorder);
-			SpawnRecorder->CreateSection(ActorToRecord.Get(), MovieScene, Guid, CurrentSequenceTime, true);
-			PropertyRecorders.Add(SpawnRecorder);
-
 			// we need to create a transform track even if we arent recording transforms
-			TSharedPtr<FMovieScene3DTransformPropertyRecorder> TransformRecorder = MakeShareable(new FMovieScene3DTransformPropertyRecorder(FirstAnimRecorder));
-			TransformRecorder->CreateSection(ActorToRecord.Get(), MovieScene, Guid, CurrentSequenceTime, ActorSettings.bRecordTransforms);
-			PropertyRecorders.Add(TransformRecorder);
-
-			// add a visibility track 
-			TSharedPtr<FMovieSceneVisibilityPropertyRecorder> VisibilityRecorder = MakeShareable(new FMovieSceneVisibilityPropertyRecorder);
-			VisibilityRecorder->CreateSection(ActorToRecord.Get(), MovieScene, Guid, CurrentSequenceTime, ActorSettings.bRecordVisibility);
-			PropertyRecorders.Add(VisibilityRecorder);
-
-			// add an attachment track
-			TSharedPtr<FMovieScene3DAttachPropertyRecorder> AttachRecorder = MakeShareable(new FMovieScene3DAttachPropertyRecorder);
-			AttachRecorder->CreateSection(ActorToRecord.Get(), MovieScene, Guid, CurrentSequenceTime, true);
-			PropertyRecorders.Add(AttachRecorder);			
-		}
-	}
-}
-
-static FString GetUniquePosessableName(UMovieScene* MovieScene, const FString& BaseName)
-{
-	FString BlueprintName = BaseName;
-	auto DuplName = [&](FMovieScenePossessable& InPosessable)
-	{
-		return InPosessable.GetName() == BlueprintName;
-	};
+			if (FSequenceRecorder::Get().GetTransformRecorderFactory().CanRecordObject(ActorToRecord.Get()))
+			{
+				UMovieScene3DTransformSectionRecorderSettings* Settings = ActorSettings.GetSettingsObject<UMovieScene3DTransformSectionRecorderSettings>();
+				check(Settings);
+
+				TSharedPtr<IMovieSceneSectionRecorder> Recorder = FSequenceRecorder::Get().GetTransformRecorderFactory().CreateSectionRecorder(Settings->bRecordTransforms, FirstAnimRecorder);
+				if(Recorder.IsValid())
+				{ 
+					Recorder->CreateSection(ActorToRecord.Get(), MovieScene, Guid, CurrentSequenceTime);
+					Recorder->Record(CurrentSequenceTime);
+					SectionRecorders.Add(Recorder);
+				}
+			}
 
-	int32 Index = 2;
-	FString UniqueString;
-	while (MovieScene->FindPossessable(DuplName))
-	{
-		BlueprintName.RemoveFromEnd(UniqueString);
-		UniqueString = FString::Printf(TEXT(" (%d)"), Index++);
-		BlueprintName += UniqueString;
+			TArray<IMovieSceneSectionRecorderFactory*> ModularFeatures = IModularFeatures::Get().GetModularFeatureImplementations<IMovieSceneSectionRecorderFactory>(MovieSceneSectionRecorderFactoryName);
+			for (IMovieSceneSectionRecorderFactory* Factory : ModularFeatures)
+			{
+				if (Factory->CanRecordObject(ActorToRecord.Get()))
+				{
+					TSharedPtr<IMovieSceneSectionRecorder> Recorder = Factory->CreateSectionRecorder(ActorSettings);
+					if (Recorder.IsValid())
+					{
+						Recorder->CreateSection(ActorToRecord.Get(), MovieScene, Guid, CurrentSequenceTime);
+						Recorder->Record(CurrentSequenceTime);
+						SectionRecorders.Add(Recorder);
+					}
+				}
+			}	
+		}
 	}
-
-	return BlueprintName;
 }
 
-TSharedPtr<FMovieSceneAnimationPropertyRecorder> UActorRecording::StartRecordingComponentProperties(const FName& BindingName, USceneComponent* SceneComponent, UObject* BindingContext, ULevelSequence* CurrentSequence, float CurrentSequenceTime)
+TSharedPtr<FMovieSceneAnimationSectionRecorder> UActorRecording::StartRecordingComponentProperties(const FName& BindingName, USceneComponent* SceneComponent, UObject* BindingContext, ULevelSequence* CurrentSequence, float CurrentSequenceTime)
 {
 	// first create a possessable for this component to be controlled by
 	UMovieScene* OwnerMovieScene = CurrentSequence->GetMovieScene();
 
-	FString Name = GetUniquePosessableName(OwnerMovieScene, BindingName.ToString());
-
-	const FGuid PossessableGuid = OwnerMovieScene->AddPossessable(Name, SceneComponent->GetClass());
+	const FGuid PossessableGuid = OwnerMovieScene->AddPossessable(BindingName.ToString(), SceneComponent->GetClass());
 
 	// Set up parent/child guids for possessables within spawnables
 	FMovieScenePossessable* ChildPossessable = OwnerMovieScene->FindPossessable(PossessableGuid);
@@ -391,44 +388,49 @@ TSharedPtr<FMovieSceneAnimationPropertyRecorder> UActorRecording::StartRecording
 		ParentSpawnable->AddChildPossessable(PossessableGuid);
 	}
 
+	// BindingName must be the component's path relative to its owner Actor
 	FLevelSequenceObjectReference ObjectReference(FUniqueObjectGuid(), BindingName.ToString());
 
 	CurrentSequence->BindPossessableObject(PossessableGuid, ObjectReference);
 
-	TSharedPtr<FMovieSceneAnimationPropertyRecorder> AnimationRecorder;
-	if(SceneComponent->IsA<USkeletalMeshComponent>())
+	// First try built-in animation recorder...
+	TSharedPtr<FMovieSceneAnimationSectionRecorder> AnimationRecorder = nullptr;
+	if (FSequenceRecorder::Get().GetAnimationRecorderFactory().CanRecordObject(SceneComponent))
 	{
-		AnimationRecorder = MakeShareable(new FMovieSceneAnimationPropertyRecorder(AnimationSettings, TargetAnimation.Get()));
-		AnimationRecorder->CreateSection(SceneComponent, OwnerMovieScene, PossessableGuid, CurrentSequenceTime, true);
-		PropertyRecorders.Add(AnimationRecorder);
-	}
-	else if(SceneComponent->IsA<UParticleSystemComponent>())
-	{
-		TSharedPtr<FMovieSceneParticleTrackPropertyRecorder> ParticleTrackRecorder = MakeShareable(new FMovieSceneParticleTrackPropertyRecorder());
-		ParticleTrackRecorder->CreateSection(SceneComponent, OwnerMovieScene, PossessableGuid, CurrentSequenceTime, true);
-		PropertyRecorders.Add(ParticleTrackRecorder);
+		AnimationRecorder = FSequenceRecorder::Get().GetAnimationRecorderFactory().CreateSectionRecorder(this);
+		AnimationRecorder->CreateSection(SceneComponent, OwnerMovieScene, PossessableGuid, CurrentSequenceTime);
+		AnimationRecorder->Record(CurrentSequenceTime);
+		SectionRecorders.Add(AnimationRecorder);
 	}
 
-	// Dont record the root component transforms as this will be taken into account by the actor transform track
-	// Also dont record transforms of skeletal mesh components as they will be taken into account in the actor transform
-	bool bIsCharacterSkelMesh = false;
-	if(SceneComponent->IsA<USkeletalMeshComponent>() && SceneComponent->GetOwner()->IsA<ACharacter>())
+	// ...and transform...
+	if (FSequenceRecorder::Get().GetTransformRecorderFactory().CanRecordObject(SceneComponent))
 	{
-		ACharacter* Character = CastChecked<ACharacter>(SceneComponent->GetOwner());
-		bIsCharacterSkelMesh = SceneComponent == Character->GetMesh();
+		TSharedPtr<IMovieSceneSectionRecorder> Recorder = FSequenceRecorder::Get().GetTransformRecorderFactory().CreateSectionRecorder(true, nullptr);
+		if (Recorder.IsValid())
+		{
+			Recorder->CreateSection(SceneComponent, OwnerMovieScene, PossessableGuid, CurrentSequenceTime);
+			Recorder->Record(CurrentSequenceTime);
+			SectionRecorders.Add(Recorder);
+		}
 	}
 
-	if(SceneComponent != SceneComponent->GetOwner()->GetRootComponent() && !bIsCharacterSkelMesh)
+	// ...now any external recorders
+	TArray<IMovieSceneSectionRecorderFactory*> ModularFeatures = IModularFeatures::Get().GetModularFeatureImplementations<IMovieSceneSectionRecorderFactory>(MovieSceneSectionRecorderFactoryName);
+	for (IMovieSceneSectionRecorderFactory* Factory : ModularFeatures)
 	{
-		TSharedPtr<FMovieScene3DTransformPropertyRecorder> TransformRecorder = MakeShareable(new FMovieScene3DTransformPropertyRecorder);
-		TransformRecorder->CreateSection(SceneComponent, OwnerMovieScene, PossessableGuid, CurrentSequenceTime, true);
-		PropertyRecorders.Add(TransformRecorder);
+		if (Factory->CanRecordObject(SceneComponent))
+		{
+			TSharedPtr<IMovieSceneSectionRecorder> Recorder = Factory->CreateSectionRecorder(ActorSettings);
+			if (Recorder.IsValid())
+			{
+				Recorder->CreateSection(SceneComponent, OwnerMovieScene, PossessableGuid, CurrentSequenceTime);
+				Recorder->Record(CurrentSequenceTime);
+				SectionRecorders.Add(Recorder);
+			}
+		}
 	}
 
-	TSharedPtr<FMovieSceneVisibilityPropertyRecorder> VisibilityRecorder = MakeShareable(new FMovieSceneVisibilityPropertyRecorder);
-	VisibilityRecorder->CreateSection(SceneComponent, OwnerMovieScene, PossessableGuid, CurrentSequenceTime, true);
-	PropertyRecorders.Add(VisibilityRecorder);	
-
 	return AnimationRecorder;
 }
 
@@ -448,9 +450,9 @@ void UActorRecording::Tick(float DeltaSeconds, ULevelSequence* CurrentSequence,
 			}
 		}
 
-		for(auto& PropertyRecorder : PropertyRecorders)
+		for(auto& SectionRecorder : SectionRecorders)
 		{
-			PropertyRecorder->Record(CurrentSequenceTime);
+			SectionRecorder->Record(CurrentSequenceTime);
 		}
 	}
 }
@@ -470,39 +472,26 @@ bool UActorRecording::StopRecording(ULevelSequence* CurrentSequence)
 		}
 	}
 
-	FScopedSlowTask SlowTask((float)PropertyRecorders.Num() + 1.0f, FText::Format(NSLOCTEXT("SequenceRecorder", "ProcessingActor", "Processing Actor {0}"), FText::FromString(ActorName)));
+	FScopedSlowTask SlowTask((float)SectionRecorders.Num() + 1.0f, FText::Format(NSLOCTEXT("SequenceRecorder", "ProcessingActor", "Processing Actor {0}"), FText::FromString(ActorName)));
 
 	// stop property recorders
-	for(auto& PropertyRecorder : PropertyRecorders)
+	for(auto& SectionRecorder : SectionRecorders)
 	{
 		SlowTask.EnterProgressFrame();
 
-		PropertyRecorder->FinalizeSection();
+		SectionRecorder->FinalizeSection();
 	}
 
 	SlowTask.EnterProgressFrame();
 
-	PropertyRecorders.Empty();
-
-	if(bNewComponentAddedWhileRecording && CurrentSequence)
-	{
-		UMovieScene* MovieScene = CurrentSequence->GetMovieScene();
-		check(MovieScene);
-
-		FMovieSceneSpawnable* Spawnable = MovieScene->FindSpawnable(Guid);
-		if(Spawnable)
-		{
-			UBlueprint* Blueprint = CastChecked<UBlueprint>(Spawnable->GetClass()->ClassGeneratedBy);
-			FBlueprintEditorUtils::MarkBlueprintAsStructurallyModified(Blueprint);
-		}
-	}
+	SectionRecorders.Empty();
 
 	return true;
 }
 
 bool UActorRecording::IsRecording() const
 {
-	return ActorToRecord.IsValid() && PropertyRecorders.Num() > 0;
+	return ActorToRecord.IsValid() && SectionRecorders.Num() > 0;
 }
 
 static FName FindParentComponentOwnerClassName(USceneComponent* SceneComponent, UBlueprint* Blueprint)
@@ -554,7 +543,7 @@ void UActorRecording::StartRecordingNewComponents(ULevelSequence* CurrentSequenc
 		FMovieSceneSpawnable* Spawnable = MovieScene->FindSpawnable(Guid);
 		check(Spawnable);
 
-		UBlueprint* Blueprint = CastChecked<UBlueprint>(Spawnable->GetClass()->ClassGeneratedBy);
+		AActor* ObjectTemplate = CastChecked<AActor>(Spawnable->GetObjectTemplate());
 
 		for(USceneComponent* SceneComponent : NewComponents)
 		{
@@ -562,38 +551,39 @@ void UActorRecording::StartRecordingNewComponents(ULevelSequence* CurrentSequenc
 			FName NewName;
 			if(SceneComponent->CreationMethod != EComponentCreationMethod::SimpleConstructionScript)
 			{
-				NewName = *FString::Printf(TEXT("Dynamic%s"), *SceneComponent->GetFName().ToString());
-				USCS_Node* NewSCSNode = Blueprint->SimpleConstructionScript->CreateNode(SceneComponent->GetClass(), NewName); 
-				UEditorEngine::CopyPropertiesForUnrelatedObjects(SceneComponent, NewSCSNode->ComponentTemplate);
-				NewSCSNode->ComponentTemplate->CreationMethod = EComponentCreationMethod::SimpleConstructionScript;
+				// Give this component a unique name within its parent
+				NewName = *FString::Printf(TEXT("Dynamic%s"), *SceneComponent->GetFName().GetPlainNameString());
+				NewName.SetNumber(1);
+				while (FindObjectFast<UObject>(ObjectTemplate, NewName))
+				{
+					NewName.SetNumber(NewName.GetNumber() + 1);
+				}
+
+				USceneComponent* TemplateRoot = ObjectTemplate->GetRootComponent();
+				USceneComponent* AttachToComponent = TemplateRoot;
 
 				// look for a similar attach parent in the current structure
-				bool bAdded = false;
 				if(SceneComponent->GetAttachParent() != nullptr)
 				{
 					FName AttachName = SceneComponent->GetAttachParent()->GetFName();
-					NewSCSNode->ParentComponentOrVariableName = AttachName;
-					NewSCSNode->bIsParentComponentNative = SceneComponent->GetAttachParent()->CreationMethod == EComponentCreationMethod::Native;
-					NewSCSNode->ParentComponentOwnerClassName = FindParentComponentOwnerClassName(SceneComponent, Blueprint);
 
-					USCS_Node* ParentSCSNode = Blueprint->SimpleConstructionScript->FindSCSNode(AttachName);
-					if(ParentSCSNode)
+					TArray<USceneComponent*> AllChildren;
+					TemplateRoot->GetChildrenComponents(true, AllChildren);
+
+					for (USceneComponent* Child : AllChildren)
 					{
-						ParentSCSNode->AddChildNode(NewSCSNode);
-						bAdded = true;
+						if (Child->GetFName() == AttachName)
+						{
+							AttachToComponent = Child;
+							break;
+						}
 					}
 				}
-			
-				// nothing found - add to root
-				if(!bAdded)
-				{
-					Blueprint->SimpleConstructionScript->AddNode(NewSCSNode);
-				}
 
-				if(SceneComponent->GetAttachSocketName() != NAME_None)
-				{
-					NewSCSNode->AttachToName = SceneComponent->GetAttachSocketName();
-				}
+				USceneComponent* NewTemplateComponent = DuplicateObject<USceneComponent>(SceneComponent, ObjectTemplate, NewName);
+				NewTemplateComponent->AttachToComponent(AttachToComponent, FAttachmentTransformRules::KeepRelativeTransform, SceneComponent->GetAttachSocketName());
+
+				ObjectTemplate->AddInstanceComponent(NewTemplateComponent);
 			}
 			else
 			{
diff --git a/Engine/Source/Editor/SequenceRecorder/Private/ActorRecording.h b/Engine/Source/Editor/SequenceRecorder/Private/ActorRecording.h
--- a/Engine/Source/Editor/SequenceRecorder/Private/ActorRecording.h
+++ b/Engine/Source/Editor/SequenceRecorder/Private/ActorRecording.h
@@ -4,32 +4,10 @@
 
 #include "Animation/AnimationRecordingSettings.h"
 #include "Components/SkinnedMeshComponent.h"
+#include "ActorRecordingSettings.h"
 
 #include "ActorRecording.generated.h"
 
-USTRUCT()
-struct FActorRecordingSettings
-{
-	GENERATED_BODY()
-
-	FActorRecordingSettings()
-		: bRecordTransforms(false)
-		, bRecordVisibility(true)
-	{}
-
-public:
-	/** 
-	 * Whether to record actor transforms. This can be useful if you want the actors to end up in specific locations after the sequence. 
-	 * By default we rely on animations to provide transforms, but this can be changed using the "Record In World Space" animation setting.
-	 */
-	UPROPERTY(EditAnywhere, Category = "Actor Recording")
-	bool bRecordTransforms;
-
-	/** Whether to record actor visibility. */
-	UPROPERTY(EditAnywhere, Category = "Actor Recording")
-	bool bRecordVisibility;
-};
-
 UCLASS(MinimalAPI, Transient)
 class UActorRecording : public UObject
 {
@@ -54,8 +32,14 @@ public:
 	/** Simulate a de-spawned actor */
 	void InvalidateObjectToRecord();
 
+	/** Get the Guid that identifies our spawnable in a recorded sequence */
+	const FGuid& GetSpawnableGuid() const
+	{
+		return Guid;
+	}
+
 private:
-	/** Check component valididty for recording */
+	/** Check component validity for recording */
 	bool ValidComponent(USceneComponent* SceneComponent) const;
 
 	/** Adds us to a folder for better sequence organization */
@@ -65,7 +49,7 @@ private:
 	void StartRecordingActorProperties(ULevelSequence* CurrentSequence, float CurrentSequenceTime);
 
 	/** Start recording component properties to a sequence */
-	TSharedPtr<class FMovieSceneAnimationPropertyRecorder> StartRecordingComponentProperties(const FName& BindingName, USceneComponent* SceneComponent, UObject* BindingContext, ULevelSequence* CurrentSequence, float CurrentSequenceTime);
+	TSharedPtr<class FMovieSceneAnimationSectionRecorder> StartRecordingComponentProperties(const FName& BindingName, USceneComponent* SceneComponent, UObject* BindingContext, ULevelSequence* CurrentSequence, float CurrentSequenceTime);
 
 	/** Start recording components that are added at runtime */
 	void StartRecordingNewComponents(ULevelSequence* CurrentSequence, float CurrentSequenceTime);
@@ -96,10 +80,7 @@ public:
 	UPROPERTY(EditAnywhere, Category = "Animation Recording")
 	FAnimationRecordingSettings AnimationSettings;
 
-	/** Guid that identifies our spawnable in a recorded sequence */
-	FGuid Guid;
-
-	/** Whether this actor recording was triggered from an acotr spawn */
+	/** Whether this actor recording was triggered from an actor spawn */
 	bool bWasSpawnedPostRecord;
 
 private:
@@ -109,12 +90,15 @@ private:
 	/** Used to store/restore URO when recording */
 	bool bEnableUpdateRateOptimizations;
 
-	/** This actor's current set of property recorders */
-	TArray<TSharedPtr<class IMovieScenePropertyRecorder>> PropertyRecorders;
+	/** This actor's current set of section recorders */
+	TArray<TSharedPtr<class IMovieSceneSectionRecorder>> SectionRecorders;
 
 	/** Track components to check if any have changed */
 	TArray<TWeakObjectPtr<USceneComponent>> TrackedComponents;
 
 	/** Flag to track whether we created new components */
 	bool bNewComponentAddedWhileRecording;
+
+	/** Guid that identifies our spawnable in a recorded sequence */
+	FGuid Guid;
 };
\ No newline at end of file
diff --git a/Engine/Source/Editor/SequenceRecorder/Private/ActorRecordingSettings.cpp b/Engine/Source/Editor/SequenceRecorder/Private/ActorRecordingSettings.cpp
new file mode 100644
--- /dev/null
+++ b/Engine/Source/Editor/SequenceRecorder/Private/ActorRecordingSettings.cpp
@@ -0,0 +1,18 @@
+// Copyright 1998-2016 Epic Games, Inc. All Rights Reserved.
+
+#include "SequenceRecorderPrivatePCH.h"
+#include "ActorRecordingSettings.h"
+#include "Runtime/Core/Public/Features/IModularFeatures.h"
+
+FActorRecordingSettings::FActorRecordingSettings()
+{
+	TArray<IMovieSceneSectionRecorderFactory*> ModularFeatures = IModularFeatures::Get().GetModularFeatureImplementations<IMovieSceneSectionRecorderFactory>("MovieSceneSectionRecorderFactory");
+	for (IMovieSceneSectionRecorderFactory* Factory : ModularFeatures)
+	{
+		UObject* SettingsObject = Factory->CreateSettingsObject();
+		if (SettingsObject)
+		{
+			Settings.Add(SettingsObject);
+		}
+	}
+}
\ No newline at end of file
diff --git a/Engine/Source/Editor/SequenceRecorder/Private/AnimationRecorder.cpp b/Engine/Source/Editor/SequenceRecorder/Private/AnimationRecorder.cpp
--- a/Engine/Source/Editor/SequenceRecorder/Private/AnimationRecorder.cpp
+++ b/Engine/Source/Editor/SequenceRecorder/Private/AnimationRecorder.cpp
@@ -470,7 +470,7 @@ void FAnimationRecorder::UpdateRecord(USkeletalMeshComponent* Component, float D
 			BlendedComponentToWorld.Blend(PreviousComponentToWorld, Component->ComponentToWorld, BlendAlpha);
 
 			FBlendedHeapCurve BlendedCurve;
-			if (PreviousAnimCurves.Elements.Num() == AnimCurves.Elements.Num())
+			if (AnimCurves.Elements.Num() > 0 && PreviousAnimCurves.Elements.Num() == AnimCurves.Elements.Num())
 			{
 				BlendedCurve.Blend(PreviousAnimCurves, AnimCurves, BlendAlpha);
 			}
diff --git a/Engine/Source/Editor/SequenceRecorder/Private/Sections/MovieScene3DTransformSectionRecorder.cpp b/Engine/Source/Editor/SequenceRecorder/Private/Sections/MovieScene3DTransformSectionRecorder.cpp
new file mode 100644
--- /dev/null
+++ b/Engine/Source/Editor/SequenceRecorder/Private/Sections/MovieScene3DTransformSectionRecorder.cpp
@@ -0,0 +1,300 @@
+// Copyright 1998-2015 Epic Games, Inc. All Rights Reserved.
+
+#include "SequenceRecorderPrivatePCH.h"
+#include "MovieScene3DTransformSectionRecorder.h"
+#include "MovieScene.h"
+#include "MovieScene3DTransformSection.h"
+#include "MovieScene3DTransformTrack.h"
+#include "MovieSceneAnimationSectionRecorder.h"
+#include "SequenceRecorderUtils.h"
+#include "ActorRecording.h"
+#include "ActorRecordingSettings.h"
+
+TSharedPtr<IMovieSceneSectionRecorder> FMovieScene3DTransformSectionRecorderFactory::CreateSectionRecorder(const FActorRecordingSettings& InActorRecordingSettings) const
+{
+	return nullptr;
+}
+
+TSharedPtr<FMovieScene3DTransformSectionRecorder> FMovieScene3DTransformSectionRecorderFactory::CreateSectionRecorder(bool bRecordTransforms, TSharedPtr<class FMovieSceneAnimationSectionRecorder> InAnimRecorder) const
+{
+	return MakeShareable(new FMovieScene3DTransformSectionRecorder(bRecordTransforms, InAnimRecorder));
+}
+
+bool FMovieScene3DTransformSectionRecorderFactory::CanRecordObject(UObject* InObjectToRecord) const
+{
+	if(USceneComponent* SceneComponent = Cast<USceneComponent>(InObjectToRecord))
+	{
+		// Dont record the root component transforms as this will be taken into account by the actor transform track
+		// Also dont record transforms of skeletal mesh components as they will be taken into account in the actor transform
+		bool bIsCharacterSkelMesh = false;
+		if (SceneComponent->IsA<USkeletalMeshComponent>() && SceneComponent->GetOwner()->IsA<ACharacter>())
+		{
+			ACharacter* Character = CastChecked<ACharacter>(SceneComponent->GetOwner());
+			bIsCharacterSkelMesh = SceneComponent == Character->GetMesh();
+		}
+
+		return (SceneComponent != SceneComponent->GetOwner()->GetRootComponent() && !bIsCharacterSkelMesh);
+	}
+	else 
+	{
+		return InObjectToRecord->IsA<AActor>();
+	}
+}
+
+void FMovieScene3DTransformSectionRecorder::CreateSection(UObject* InObjectToRecord, UMovieScene* InMovieScene, const FGuid& Guid, float Time) 
+{
+	ObjectToRecord = InObjectToRecord;
+	bWasAttached = false;
+		
+	MovieScene = InMovieScene;
+	MovieSceneTrack = InMovieScene->AddTrack<UMovieScene3DTransformTrack>(Guid);
+	if(MovieSceneTrack.IsValid())
+	{
+		MovieSceneSection = Cast<UMovieScene3DTransformSection>(MovieSceneTrack->CreateNewSection());
+
+		MovieSceneTrack->AddSection(*MovieSceneSection);
+
+		const bool bUnwindRotation = false;
+
+		DefaultTransform = GetTransformToRecord();
+		FVector EulerRotation = DefaultTransform.GetRotation().Rotator().Euler();
+
+		MovieSceneSection->SetDefault(FTransformKey(EKey3DTransformChannel::Translation, EAxis::X, DefaultTransform.GetTranslation().X, bUnwindRotation));
+		MovieSceneSection->SetDefault(FTransformKey(EKey3DTransformChannel::Translation, EAxis::Y, DefaultTransform.GetTranslation().Y, bUnwindRotation));
+		MovieSceneSection->SetDefault(FTransformKey(EKey3DTransformChannel::Translation, EAxis::Z, DefaultTransform.GetTranslation().Z, bUnwindRotation));
+
+		MovieSceneSection->SetDefault(FTransformKey(EKey3DTransformChannel::Rotation, EAxis::X, EulerRotation.X, bUnwindRotation));
+		MovieSceneSection->SetDefault(FTransformKey(EKey3DTransformChannel::Rotation, EAxis::Y, EulerRotation.Y, bUnwindRotation));
+		MovieSceneSection->SetDefault(FTransformKey(EKey3DTransformChannel::Rotation, EAxis::Z, EulerRotation.Z, bUnwindRotation));
+
+		MovieSceneSection->SetDefault(FTransformKey(EKey3DTransformChannel::Scale, EAxis::X, DefaultTransform.GetScale3D().X, bUnwindRotation));
+		MovieSceneSection->SetDefault(FTransformKey(EKey3DTransformChannel::Scale, EAxis::Y, DefaultTransform.GetScale3D().Y, bUnwindRotation));
+		MovieSceneSection->SetDefault(FTransformKey(EKey3DTransformChannel::Scale, EAxis::Z, DefaultTransform.GetScale3D().Z, bUnwindRotation));
+
+		MovieSceneSection->SetStartTime(Time);
+		MovieSceneSection->SetIsInfinite(true);
+	}
+}
+
+void FMovieScene3DTransformSectionRecorder::FinalizeSection()
+{
+	FScopedSlowTask SlowTask(4.0f, NSLOCTEXT("SequenceRecorder", "ProcessingTransforms", "Processing Transforms"));
+
+	bRecording = false;
+
+	// if we have a valid animation recorder, we need to build our transforms from the animation
+	// so we properly synchronize our keyframes
+	if(AnimRecorder.IsValid())
+	{
+		check(BufferedTransforms.Num() == 0);
+
+		UAnimSequence* AnimSequence = AnimRecorder->GetAnimSequence();
+		USkeletalMesh* SkeletalMesh = AnimRecorder->GetSkeletalMesh();
+		if(AnimSequence && SkeletalMesh)
+		{
+			// find the root bone
+			int32 RootIndex = INDEX_NONE;
+			USkeleton* AnimSkeleton = AnimSequence->GetSkeleton();
+			for (int32 TrackIndex = 0; TrackIndex < AnimSequence->RawAnimationData.Num(); ++TrackIndex)
+			{
+				// verify if this bone exists in skeleton
+				int32 BoneTreeIndex = AnimSequence->GetSkeletonIndexFromRawDataTrackIndex(TrackIndex);
+				if (BoneTreeIndex != INDEX_NONE)
+				{
+					int32 BoneIndex = AnimSkeleton->GetMeshBoneIndexFromSkeletonBoneIndex(SkeletalMesh, BoneTreeIndex);
+					int32 ParentIndex = SkeletalMesh->RefSkeleton.GetParentIndex(BoneIndex);
+					if(ParentIndex == INDEX_NONE)
+					{
+						// found root
+						RootIndex = BoneIndex;
+						break;
+					}
+				}
+			}
+
+			check(RootIndex != INDEX_NONE);
+
+			const float StartTime = MovieSceneSection->GetStartTime();
+
+			// we may need to offset the transform here if the animation was not recorded on the root component
+			FTransform InvComponentTransform = AnimRecorder->GetComponentTransform().Inverse();
+
+			FRawAnimSequenceTrack& RawTrack = AnimSequence->RawAnimationData[RootIndex];
+			const int32 KeyCount = FMath::Max(FMath::Max(RawTrack.PosKeys.Num(), RawTrack.RotKeys.Num()), RawTrack.ScaleKeys.Num());
+			for(int32 KeyIndex = 0; KeyIndex < KeyCount; KeyIndex++)
+			{
+				FTransform Transform;
+				if(RawTrack.PosKeys.IsValidIndex(KeyIndex))
+				{
+					Transform.SetTranslation(RawTrack.PosKeys[KeyIndex]);
+				}
+				else if(RawTrack.PosKeys.Num() > 0)
+				{
+					Transform.SetTranslation(RawTrack.PosKeys[0]);
+				}
+				
+				if(RawTrack.RotKeys.IsValidIndex(KeyIndex))
+				{
+					Transform.SetRotation(RawTrack.RotKeys[KeyIndex]);
+				}
+				else if(RawTrack.RotKeys.Num() > 0)
+				{
+					Transform.SetRotation(RawTrack.RotKeys[0]);
+				}
+
+				if(RawTrack.ScaleKeys.IsValidIndex(KeyIndex))
+				{
+					Transform.SetScale3D(RawTrack.ScaleKeys[KeyIndex]);
+				}
+				else if(RawTrack.ScaleKeys.Num() > 0)
+				{
+					Transform.SetScale3D(RawTrack.ScaleKeys[0]);
+				}
+
+				BufferedTransforms.Add(FBufferedTransformKey(InvComponentTransform * Transform, StartTime + AnimSequence->GetTimeAtFrame(KeyIndex)));
+			}
+		}
+	}
+
+	SlowTask.EnterProgressFrame();
+
+	// Try to 're-wind' rotations that look like axis flips
+	// We need to do this as a post-process because the recorder cant reliably access 'wound' rotations:
+	// - Net quantize may use quaternions.
+	// - Scene components cache transforms as quaternions.
+	// - Gameplay is free to clamp/fmod rotations as it sees fit.
+	int32 TransformCount = BufferedTransforms.Num();
+	for(int32 TransformIndex = 0; TransformIndex < TransformCount - 1; TransformIndex++)
+	{
+		FRotator& Rotator = BufferedTransforms[TransformIndex].WoundRotation;
+		FRotator& NextRotator = BufferedTransforms[TransformIndex + 1].WoundRotation;
+
+		FMath::WindRelativeAnglesDegrees(Rotator.Pitch, NextRotator.Pitch);
+		FMath::WindRelativeAnglesDegrees(Rotator.Yaw, NextRotator.Yaw);
+		FMath::WindRelativeAnglesDegrees(Rotator.Roll, NextRotator.Roll);
+	}
+
+	SlowTask.EnterProgressFrame();
+
+	// never unwind rotations
+	const bool bUnwindRotation = false;
+	// If we are syncing to an animation, use linear interpolation to avoid foot sliding etc. 
+	// Otherwise use cubic for better quality (much better for projectiles etc.)
+	const EMovieSceneKeyInterpolation Interpolation = AnimRecorder.IsValid() ? EMovieSceneKeyInterpolation::Linear : EMovieSceneKeyInterpolation::Auto;
+
+	// add buffered transforms
+	for(const FBufferedTransformKey& BufferedTransform : BufferedTransforms)
+	{
+		const FVector Translation = BufferedTransform.Transform.GetTranslation();
+		const FVector Rotation = BufferedTransform.WoundRotation.Euler();
+		const FVector Scale = BufferedTransform.Transform.GetScale3D();
+
+		MovieSceneSection->AddKey(BufferedTransform.KeyTime, FTransformKey(EKey3DTransformChannel::Translation, EAxis::X, Translation.X, bUnwindRotation), Interpolation);
+		MovieSceneSection->AddKey(BufferedTransform.KeyTime, FTransformKey(EKey3DTransformChannel::Translation, EAxis::Y, Translation.Y, bUnwindRotation), Interpolation);
+		MovieSceneSection->AddKey(BufferedTransform.KeyTime, FTransformKey(EKey3DTransformChannel::Translation, EAxis::Z, Translation.Z, bUnwindRotation), Interpolation);
+
+		MovieSceneSection->AddKey(BufferedTransform.KeyTime, FTransformKey(EKey3DTransformChannel::Rotation, EAxis::X, Rotation.X, bUnwindRotation), Interpolation);
+		MovieSceneSection->AddKey(BufferedTransform.KeyTime, FTransformKey(EKey3DTransformChannel::Rotation, EAxis::Y, Rotation.Y, bUnwindRotation), Interpolation);
+		MovieSceneSection->AddKey(BufferedTransform.KeyTime, FTransformKey(EKey3DTransformChannel::Rotation, EAxis::Z, Rotation.Z, bUnwindRotation), Interpolation);
+
+		MovieSceneSection->AddKey(BufferedTransform.KeyTime, FTransformKey(EKey3DTransformChannel::Scale, EAxis::X, Scale.X, bUnwindRotation), Interpolation);
+		MovieSceneSection->AddKey(BufferedTransform.KeyTime, FTransformKey(EKey3DTransformChannel::Scale, EAxis::Y, Scale.Y, bUnwindRotation), Interpolation);
+		MovieSceneSection->AddKey(BufferedTransform.KeyTime, FTransformKey(EKey3DTransformChannel::Scale, EAxis::Z, Scale.Z, bUnwindRotation), Interpolation);
+	}
+
+	BufferedTransforms.Empty();
+
+	SlowTask.EnterProgressFrame();
+
+	// now remove linear keys
+	TPair<FRichCurve*, float> CurvesAndTolerances[] =
+	{
+		TPairInitializer<FRichCurve*, float>(&MovieSceneSection->GetTranslationCurve(EAxis::X), KINDA_SMALL_NUMBER),
+		TPairInitializer<FRichCurve*, float>(&MovieSceneSection->GetTranslationCurve(EAxis::Y), KINDA_SMALL_NUMBER),
+		TPairInitializer<FRichCurve*, float>(&MovieSceneSection->GetTranslationCurve(EAxis::Z), KINDA_SMALL_NUMBER),
+		TPairInitializer<FRichCurve*, float>(&MovieSceneSection->GetRotationCurve(EAxis::X), KINDA_SMALL_NUMBER),
+		TPairInitializer<FRichCurve*, float>(&MovieSceneSection->GetRotationCurve(EAxis::Y), KINDA_SMALL_NUMBER),
+		TPairInitializer<FRichCurve*, float>(&MovieSceneSection->GetRotationCurve(EAxis::Z), KINDA_SMALL_NUMBER),
+		TPairInitializer<FRichCurve*, float>(&MovieSceneSection->GetScaleCurve(EAxis::X), KINDA_SMALL_NUMBER),
+		TPairInitializer<FRichCurve*, float>(&MovieSceneSection->GetScaleCurve(EAxis::Y), KINDA_SMALL_NUMBER),
+		TPairInitializer<FRichCurve*, float>(&MovieSceneSection->GetScaleCurve(EAxis::Z), KINDA_SMALL_NUMBER),
+	};
+
+	for(TPair<FRichCurve*, float>& CurveAndTolerance : CurvesAndTolerances)
+	{
+		CurveAndTolerance.Key->RemoveRedundantKeys(CurveAndTolerance.Value);
+	}
+
+	// we cant remove redundant tracks if we were attached as the playback relies on update order of
+	// transform tracks. Without this track, relative transforms would accumulate.
+	if(!bWasAttached)
+	{
+		// now we have reduced our keys, if we dont have any, remove the section as it is redundant
+		if( MovieSceneSection->GetTranslationCurve(EAxis::X).Keys.Num() == 0 && 
+			MovieSceneSection->GetTranslationCurve(EAxis::Y).Keys.Num() == 0 && 
+			MovieSceneSection->GetTranslationCurve(EAxis::Z).Keys.Num() == 0 && 
+			MovieSceneSection->GetRotationCurve(EAxis::X).Keys.Num() == 0 && 
+			MovieSceneSection->GetRotationCurve(EAxis::Y).Keys.Num() == 0 && 
+			MovieSceneSection->GetRotationCurve(EAxis::Z).Keys.Num() == 0 && 
+			MovieSceneSection->GetScaleCurve(EAxis::X).Keys.Num() == 0 && 
+			MovieSceneSection->GetScaleCurve(EAxis::Y).Keys.Num() == 0 && 
+			MovieSceneSection->GetScaleCurve(EAxis::Z).Keys.Num() == 0)
+		{
+			if(DefaultTransform.Equals(FTransform::Identity))
+			{
+				MovieScene->RemoveTrack(*MovieSceneTrack.Get());
+			}
+		}
+	}
+
+	SlowTask.EnterProgressFrame();
+}
+
+void FMovieScene3DTransformSectionRecorder::Record(float CurrentTime)
+{
+	if(ObjectToRecord.IsValid())
+	{
+		if(USceneComponent* SceneComponent = Cast<USceneComponent>(ObjectToRecord.Get()))
+		{
+			// dont record non-registered scene components
+			if(!SceneComponent->IsRegistered())
+			{
+				return;
+			}
+		}
+
+		MovieSceneSection->SetEndTime(CurrentTime);
+
+		if(bRecording)
+		{
+			// don't record from the transform of the component/actor if we are synchronizing with an animation
+			if(!AnimRecorder.IsValid())
+			{
+				BufferedTransforms.Add(FBufferedTransformKey(GetTransformToRecord(), CurrentTime));
+			}
+		}
+	}
+}
+
+FTransform FMovieScene3DTransformSectionRecorder::GetTransformToRecord()
+{
+	if(USceneComponent* SceneComponent = Cast<USceneComponent>(ObjectToRecord.Get()))
+	{
+		return SceneComponent->GetRelativeTransform();
+	}
+	else if(AActor* Actor = Cast<AActor>(ObjectToRecord.Get()))
+	{
+		FName SocketName, ComponentName;
+		if(SequenceRecorderUtils::GetAttachment(Actor, SocketName, ComponentName) != nullptr && Actor->GetRootComponent() != nullptr)
+		{
+			bWasAttached = true;
+			return Actor->GetRootComponent()->GetRelativeTransform();
+		}
+		else
+		{
+			return Actor->GetTransform();
+		}
+	}
+
+	return FTransform::Identity;
+}
\ No newline at end of file
diff --git a/Engine/Source/Editor/SequenceRecorder/Private/Sections/MovieScene3DTransformSectionRecorder.h b/Engine/Source/Editor/SequenceRecorder/Private/Sections/MovieScene3DTransformSectionRecorder.h
new file mode 100644
--- /dev/null
+++ b/Engine/Source/Editor/SequenceRecorder/Private/Sections/MovieScene3DTransformSectionRecorder.h
@@ -0,0 +1,88 @@
+// Copyright 1998-2015 Epic Games, Inc. All Rights Reserved.
+
+#pragma once
+
+#include "MovieScene3DTransformSectionRecorderSettings.h"
+
+class FMovieScene3DTransformSectionRecorderFactory : public IMovieSceneSectionRecorderFactory
+{
+public:
+	virtual ~FMovieScene3DTransformSectionRecorderFactory() {}
+
+	virtual bool CanRecordObject(class UObject* InObjectToRecord) const override;
+	virtual UObject* CreateSettingsObject() const override { return NewObject<UMovieScene3DTransformSectionRecorderSettings>(); }
+
+	TSharedPtr<class FMovieScene3DTransformSectionRecorder> CreateSectionRecorder(bool bRecordTransforms, TSharedPtr<class FMovieSceneAnimationSectionRecorder> InAnimRecorder) const;
+
+private:
+	virtual TSharedPtr<IMovieSceneSectionRecorder> CreateSectionRecorder(const struct FActorRecordingSettings& InActorRecordingSettings) const override;
+};
+
+/** Structure used to buffer up transform keys. Keys are inserted into tracks in FinalizeSection() */
+struct FBufferedTransformKey
+{
+	FBufferedTransformKey(const FTransform& InTransform, float InKeyTime)
+		: Transform(InTransform)
+		, WoundRotation(InTransform.Rotator())
+		, KeyTime(InKeyTime)
+	{
+	}
+
+	FTransform Transform;
+	FRotator WoundRotation;
+	float KeyTime;
+};
+
+class FMovieScene3DTransformSectionRecorder : public IMovieSceneSectionRecorder
+{
+public:
+	FMovieScene3DTransformSectionRecorder(bool bInActuallyRecord, TSharedPtr<class FMovieSceneAnimationSectionRecorder> InAnimRecorder = nullptr) 
+		: bRecording(bInActuallyRecord)
+		, AnimRecorder(InAnimRecorder)
+	{}
+
+	virtual ~FMovieScene3DTransformSectionRecorder() {}
+
+	virtual void CreateSection(UObject* InObjectToRecord, class UMovieScene* InMovieScene, const FGuid& Guid, float Time) override;
+	virtual void FinalizeSection() override;
+	virtual void Record(float CurrentTime) override;
+	virtual void InvalidateObjectToRecord() override
+	{
+		ObjectToRecord = nullptr;
+	}
+	virtual UObject* GetSourceObject() const override
+	{
+		return ObjectToRecord.Get();
+	}
+
+private:
+	FTransform GetTransformToRecord();
+
+private:
+	/** Object to record from */
+	TLazyObjectPtr<UObject> ObjectToRecord;
+
+	/** MovieScene to record to */
+	TWeakObjectPtr<class UMovieScene> MovieScene;
+
+	/** Track to record to */
+	TWeakObjectPtr<class UMovieScene3DTransformTrack> MovieSceneTrack;
+
+	/** Section to record to */
+	TWeakObjectPtr<class UMovieScene3DTransformSection> MovieSceneSection;
+
+	/** Buffer of transform keys. Keys are inserted into tracks in FinalizeSection() */
+	TArray<FBufferedTransformKey> BufferedTransforms;
+
+	/** Flag if we are actually recording or not */
+	bool bRecording;
+
+	/** Animation recorder we use to sync our transforms to */
+	TSharedPtr<class FMovieSceneAnimationSectionRecorder> AnimRecorder;
+
+	/** The default transform this recording starts with */
+	FTransform DefaultTransform;
+
+	/** Flag indicating that some time while this recorder was active an attachment was also in place */
+	bool bWasAttached;
+};
\ No newline at end of file
diff --git a/Engine/Source/Editor/SequenceRecorder/Private/Sections/MovieSceneAnimationSectionRecorder.cpp b/Engine/Source/Editor/SequenceRecorder/Private/Sections/MovieSceneAnimationSectionRecorder.cpp
new file mode 100644
--- /dev/null
+++ b/Engine/Source/Editor/SequenceRecorder/Private/Sections/MovieSceneAnimationSectionRecorder.cpp
@@ -0,0 +1,144 @@
+// Copyright 1998-2015 Epic Games, Inc. All Rights Reserved.
+
+#include "SequenceRecorderPrivatePCH.h"
+#include "MovieSceneAnimationSectionRecorder.h"
+#include "MovieSceneSkeletalAnimationTrack.h"
+#include "MovieSceneSkeletalAnimationSection.h"
+#include "MovieScene.h"
+#include "SequenceRecorderUtils.h"
+#include "SequenceRecorderSettings.h"
+#include "Runtime/AssetRegistry/Public/AssetRegistryModule.h"
+#include "ActorRecording.h"
+
+TSharedPtr<IMovieSceneSectionRecorder> FMovieSceneAnimationSectionRecorderFactory::CreateSectionRecorder(const FActorRecordingSettings& InActorRecordingSettings) const
+{
+	return nullptr;
+}
+
+TSharedPtr<FMovieSceneAnimationSectionRecorder> FMovieSceneAnimationSectionRecorderFactory::CreateSectionRecorder(UActorRecording* InActorRecording) const
+{
+	return MakeShareable(new FMovieSceneAnimationSectionRecorder(InActorRecording->AnimationSettings, InActorRecording->TargetAnimation.Get()));
+}
+
+bool FMovieSceneAnimationSectionRecorderFactory::CanRecordObject(UObject* InObjectToRecord) const
+{
+	return InObjectToRecord->IsA<USkeletalMeshComponent>();
+}
+
+FMovieSceneAnimationSectionRecorder::FMovieSceneAnimationSectionRecorder(const FAnimationRecordingSettings& InAnimationSettings, UAnimSequence* InSpecifiedSequence)
+	: AnimSequence(InSpecifiedSequence)
+	, AnimationSettings(InAnimationSettings)
+{
+}
+
+void FMovieSceneAnimationSectionRecorder::CreateSection(UObject* InObjectToRecord, UMovieScene* MovieScene, const FGuid& Guid, float Time)
+{
+	ObjectToRecord = InObjectToRecord;
+
+	AActor* Actor = nullptr;
+	SkeletalMeshComponent = Cast<USkeletalMeshComponent>(InObjectToRecord);
+	if(!SkeletalMeshComponent.IsValid())
+	{
+		Actor = Cast<AActor>(InObjectToRecord);
+		if(Actor != nullptr)
+		{
+			SkeletalMeshComponent = Actor->FindComponentByClass<USkeletalMeshComponent>();
+		}
+	}
+	else
+	{
+		Actor = SkeletalMeshComponent->GetOwner();
+	}
+
+	if(SkeletalMeshComponent.IsValid())
+	{
+		SkeletalMesh = SkeletalMeshComponent->SkeletalMesh;
+
+		// turn off URO and make sure we always update even if out of view
+		bEnableUpdateRateOptimizations = SkeletalMeshComponent->bEnableUpdateRateOptimizations;
+		MeshComponentUpdateFlag = SkeletalMeshComponent->MeshComponentUpdateFlag;
+
+		SkeletalMeshComponent->bEnableUpdateRateOptimizations = false;
+		SkeletalMeshComponent->MeshComponentUpdateFlag = EMeshComponentUpdateFlag::AlwaysTickPoseAndRefreshBones;
+
+		ComponentTransform = SkeletalMeshComponent->GetComponentToWorld().GetRelativeTransform(SkeletalMeshComponent->GetOwner()->GetTransform());
+
+		if(!AnimSequence.IsValid())
+		{
+			// build an asset path
+			const USequenceRecorderSettings* Settings = GetDefault<USequenceRecorderSettings>();
+
+			FString AssetPath = Settings->SequenceRecordingBasePath.Path;
+			if(Settings->AnimationSubDirectory.Len() > 0)
+			{
+				AssetPath /= Settings->AnimationSubDirectory;
+			}
+
+			FString AssetName = Settings->SequenceName.Len() > 0 ? Settings->SequenceName : TEXT("RecordedSequence");
+			AssetName += TEXT("_");
+			AssetName += Actor->GetActorLabel();
+
+			AnimSequence = SequenceRecorderUtils::MakeNewAsset<UAnimSequence>(AssetPath, AssetName);
+			if(AnimSequence.IsValid())
+			{
+				FAssetRegistryModule::AssetCreated(AnimSequence.Get());
+
+				// set skeleton
+				AnimSequence->SetSkeleton(SkeletalMeshComponent->SkeletalMesh->Skeleton);
+			}
+		}
+
+		if(AnimSequence.IsValid())
+		{
+			FAnimationRecorderManager::Get().RecordAnimation(SkeletalMeshComponent.Get(), AnimSequence.Get(), AnimationSettings);
+
+			if(MovieScene)
+			{
+				UMovieSceneSkeletalAnimationTrack* AnimTrack = MovieScene->AddTrack<UMovieSceneSkeletalAnimationTrack>(Guid);
+				if(AnimTrack)
+				{
+					AnimTrack->AddNewAnimation(Time, AnimSequence.Get());
+					MovieSceneSection = Cast<UMovieSceneSkeletalAnimationSection>(AnimTrack->GetAllSections()[0]);
+				}
+			}
+		}
+	}
+}
+
+void FMovieSceneAnimationSectionRecorder::FinalizeSection()
+{
+	if(AnimSequence.IsValid())
+	{
+		// enable root motion on the animation
+		AnimSequence->bEnableRootMotion = true;
+		AnimSequence->RootMotionRootLock = ERootMotionRootLock::Zero;
+	}
+
+	if(SkeletalMeshComponent.IsValid())
+	{
+		// restore update flags
+		SkeletalMeshComponent->bEnableUpdateRateOptimizations = bEnableUpdateRateOptimizations;
+		SkeletalMeshComponent->MeshComponentUpdateFlag = MeshComponentUpdateFlag;
+
+		// only show a message if we dont have a valid movie section
+		const bool bShowMessage = !MovieSceneSection.IsValid();
+		FAnimationRecorderManager::Get().StopRecordingAnimation(SkeletalMeshComponent.Get(), bShowMessage);
+	}
+
+	if(MovieSceneSection.IsValid() && AnimSequence.IsValid())
+	{
+		MovieSceneSection->SetEndTime(MovieSceneSection->GetStartTime() + AnimSequence->GetPlayLength());
+	}
+}
+
+void FMovieSceneAnimationSectionRecorder::Record(float CurrentTime)
+{
+	// The animation recorder does most of the work here
+
+	if(SkeletalMeshComponent.IsValid())
+	{
+		// re-force updates on as gameplay can sometimes turn these back off!
+		SkeletalMeshComponent->bEnableUpdateRateOptimizations = false;
+		SkeletalMeshComponent->MeshComponentUpdateFlag = EMeshComponentUpdateFlag::AlwaysTickPoseAndRefreshBones;
+	}
+}
\ No newline at end of file
diff --git a/Engine/Source/Editor/SequenceRecorder/Private/Sections/MovieSceneAnimationSectionRecorder.h b/Engine/Source/Editor/SequenceRecorder/Private/Sections/MovieSceneAnimationSectionRecorder.h
new file mode 100644
--- /dev/null
+++ b/Engine/Source/Editor/SequenceRecorder/Private/Sections/MovieSceneAnimationSectionRecorder.h
@@ -0,0 +1,65 @@
+// Copyright 1998-2015 Epic Games, Inc. All Rights Reserved.
+
+#pragma once
+
+class FMovieSceneAnimationSectionRecorderFactory : public IMovieSceneSectionRecorderFactory
+{
+public:
+	virtual ~FMovieSceneAnimationSectionRecorderFactory() {}
+
+	virtual bool CanRecordObject(class UObject* InObjectToRecord) const override;
+
+	TSharedPtr<class FMovieSceneAnimationSectionRecorder> CreateSectionRecorder(class UActorRecording* InActorRecording) const;
+
+private:
+	virtual TSharedPtr<IMovieSceneSectionRecorder> CreateSectionRecorder(const struct FActorRecordingSettings& InActorRecordingSettings) const override;
+};
+
+class FMovieSceneAnimationSectionRecorder : public IMovieSceneSectionRecorder
+{
+public:
+	FMovieSceneAnimationSectionRecorder(const FAnimationRecordingSettings& InAnimationSettings, UAnimSequence* InSpecifiedSequence);
+	virtual ~FMovieSceneAnimationSectionRecorder() {}
+
+	virtual void CreateSection(UObject* InObjectToRecord, class UMovieScene* MovieScene, const FGuid& Guid, float Time) override;
+	virtual void FinalizeSection() override;
+	virtual void Record(float CurrentTime) override;
+	virtual void InvalidateObjectToRecord() override
+	{
+		ObjectToRecord = nullptr;
+	}
+	virtual UObject* GetSourceObject() const override
+	{
+		return ObjectToRecord.Get();
+	}
+
+	UAnimSequence* GetAnimSequence() const { return AnimSequence.Get(); }
+
+	USkeletalMesh* GetSkeletalMesh() const { return SkeletalMesh.Get(); }
+
+	const FTransform& GetComponentTransform() const { return ComponentTransform; }
+
+private:
+	/** Object to record from */
+	TLazyObjectPtr<UObject> ObjectToRecord;
+
+	/** Section to record to */
+	TWeakObjectPtr<class UMovieSceneSkeletalAnimationSection> MovieSceneSection;
+
+	TWeakObjectPtr<class UAnimSequence> AnimSequence;
+
+	TWeakObjectPtr<class USkeletalMeshComponent> SkeletalMeshComponent;
+
+	TWeakObjectPtr<class USkeletalMesh> SkeletalMesh;
+
+	/** Local transform of the component we are recording */
+	FTransform ComponentTransform;
+
+	FAnimationRecordingSettings AnimationSettings;
+
+	/** Used to store/restore update flag when recording */
+	EMeshComponentUpdateFlag::Type MeshComponentUpdateFlag;
+
+	/** Used to store/restore URO when recording */
+	bool bEnableUpdateRateOptimizations;
+};
\ No newline at end of file
diff --git a/Engine/Source/Editor/SequenceRecorder/Private/SequenceRecorder.cpp b/Engine/Source/Editor/SequenceRecorder/Private/SequenceRecorder.cpp
--- a/Engine/Source/Editor/SequenceRecorder/Private/SequenceRecorder.cpp
+++ b/Engine/Source/Editor/SequenceRecorder/Private/SequenceRecorder.cpp
@@ -14,6 +14,7 @@
 #include "NotificationManager.h"
 #include "DrawDebugHelpers.h"
 #include "LevelSequenceActor.h"
+#include "Runtime/Core/Public/Features/IModularFeatures.h"
 
 #define LOCTEXT_NAMESPACE "SequenceRecorder"
 
@@ -31,10 +32,18 @@ void FSequenceRecorder::Initialize()
 	CountdownTexture->AddToRoot();
 	RecordingIndicatorTexture = LoadObject<UTexture2D>(nullptr, TEXT("/Engine/EditorResources/SequenceRecorder/RecordingIndicator.RecordingIndicator"), nullptr, LOAD_None, nullptr);
 	RecordingIndicatorTexture->AddToRoot();
+
+	// register built-in recorders
+	IModularFeatures::Get().RegisterModularFeature("MovieSceneSectionRecorderFactory", &AnimationSectionRecorderFactory);
+	IModularFeatures::Get().RegisterModularFeature("MovieSceneSectionRecorderFactory", &TransformSectionRecorderFactory);
 }
 
 void FSequenceRecorder::Shutdown()
 {
+	// unregister built-in recorders
+	IModularFeatures::Get().UnregisterModularFeature("MovieSceneSectionRecorderFactory", &AnimationSectionRecorderFactory);
+	IModularFeatures::Get().UnregisterModularFeature("MovieSceneSectionRecorderFactory", &TransformSectionRecorderFactory);
+
 	if(CountdownTexture.IsValid())
 	{
 		CountdownTexture->RemoveFromRoot();
@@ -108,7 +117,10 @@ UActorRecording* FSequenceRecorder::AddNewQueuedRecording(AActor* Actor, UAnimSe
 
 	// We always record in world space as we need animations to record root motion
 	ActorRecording->AnimationSettings.bRecordInWorldSpace = true;
-	ActorRecording->ActorSettings.bRecordTransforms = true;
+
+	UMovieScene3DTransformSectionRecorderSettings* TransformSettings = ActorRecording->ActorSettings.GetSettingsObject<UMovieScene3DTransformSectionRecorderSettings>();
+	check(TransformSettings);
+	TransformSettings->bRecordTransforms = true;
 
 	// auto-save assets in non-editor runtime
 	if(GEditor == nullptr)
diff --git a/Engine/Source/Editor/SequenceRecorder/Private/SequenceRecorder.h b/Engine/Source/Editor/SequenceRecorder/Private/SequenceRecorder.h
--- a/Engine/Source/Editor/SequenceRecorder/Private/SequenceRecorder.h
+++ b/Engine/Source/Editor/SequenceRecorder/Private/SequenceRecorder.h
@@ -3,10 +3,14 @@
 #pragma once
 
 #include "ISequenceRecorder.h"
+#include "MovieSceneAnimationSectionRecorder.h"
+#include "MovieScene3DTransformSectionRecorder.h"
 
 struct FSequenceRecorder
 {
 public:
+	static FName MovieScenePropertyRecorderFactoryName;
+
 	/** Singleton accessor */
 	static FSequenceRecorder& Get();
 
@@ -73,6 +77,18 @@ public:
 	/** Handle actors being de-spawned */
 	void HandleActorDespawned(AActor* Actor);
 
+	/** Get the built-in animation factory (as this uses special case handling) */
+	const FMovieSceneAnimationSectionRecorderFactory& GetAnimationRecorderFactory() const
+	{
+		return AnimationSectionRecorderFactory;
+	}
+
+	/** Get the built-in transform factory (as this uses special case handling) */
+	const FMovieScene3DTransformSectionRecorderFactory& GetTransformRecorderFactory() const
+	{
+		return TransformSectionRecorderFactory;
+	}
+
 private:
 	/** Starts recording a sequence, possibly delayed */
 	bool StartRecordingInternal(UWorld* World);
@@ -130,4 +146,10 @@ private:
 
 	/** Cached sequence path to record to */
 	FString PathToRecordTo;
+
+	/** Built-in animation recorder factory */
+	FMovieSceneAnimationSectionRecorderFactory AnimationSectionRecorderFactory;
+
+	/** Built-in transform recorder factory */
+	FMovieScene3DTransformSectionRecorderFactory TransformSectionRecorderFactory;
 };
diff --git a/Engine/Source/Editor/SequenceRecorder/Private/SequenceRecorderModule.cpp b/Engine/Source/Editor/SequenceRecorder/Private/SequenceRecorderModule.cpp
--- a/Engine/Source/Editor/SequenceRecorder/Private/SequenceRecorderModule.cpp
+++ b/Engine/Source/Editor/SequenceRecorder/Private/SequenceRecorderModule.cpp
@@ -10,6 +10,7 @@
 #include "ISequenceRecorder.h"
 #include "NotificationManager.h"
 #include "SNotificationList.h"
+#include "ActorRecordingDetailsCustomization.h"
 
 #define LOCTEXT_NAMESPACE "SequenceRecorder"
 
@@ -97,6 +98,10 @@ class FSequenceRecorderModule : public ISequenceRecorder, private FSelfRegisteri
 
 			// register for debug drawing
 			DrawDebugDelegateHandle = UDebugDrawService::Register(TEXT("Game"),  FDebugDrawDelegate::CreateStatic(&FSequenceRecorderModule::DrawDebug));
+
+			// register details customization
+			FPropertyEditorModule& PropertyModule = FModuleManager::LoadModuleChecked<FPropertyEditorModule>("PropertyEditor");
+			PropertyModule.RegisterCustomClassLayout(UActorRecording::StaticClass()->GetFName(), FOnGetDetailCustomizationInstance::CreateStatic(&FActorRecordingDetailsCustomization::MakeInstance));
 		}
 #endif
 	}
@@ -134,6 +139,12 @@ class FSequenceRecorderModule : public ISequenceRecorder, private FSelfRegisteri
 				PersonaModule.OnGetCurrentRecording().Unbind();
 				PersonaModule.OnGetCurrentRecordingTime().Unbind();
 			}
+
+			if (FModuleManager::Get().IsModuleLoaded(TEXT("PropertyEditor")))
+			{
+				FPropertyEditorModule& PropertyModule = FModuleManager::LoadModuleChecked<FPropertyEditorModule>("PropertyEditor");
+				PropertyModule.UnregisterCustomClassLayout(UActorRecording::StaticClass()->GetFName());
+			}
 		}
 
 		if(GEngine)
@@ -508,6 +519,17 @@ class FSequenceRecorderModule : public ISequenceRecorder, private FSelfRegisteri
 		FSequenceRecorder::Get().HandleActorDespawned(Actor);
 	}
 
+	virtual FGuid GetRecordingGuid(AActor* Actor) const
+	{
+		UActorRecording* Recording = FSequenceRecorder::Get().FindRecording(Actor);
+		if (Recording != nullptr)
+		{
+			return Recording->GetSpawnableGuid();
+		}
+
+		return FGuid();
+	}
+
 	static void TickSequenceRecorder(float DeltaSeconds)
 	{
 		if (!IsRunningDedicatedServer() && !IsRunningCommandlet())
diff --git a/Engine/Source/Editor/SequenceRecorder/Private/SequenceRecorderPrivatePCH.h b/Engine/Source/Editor/SequenceRecorder/Private/SequenceRecorderPrivatePCH.h
--- a/Engine/Source/Editor/SequenceRecorder/Private/SequenceRecorderPrivatePCH.h
+++ b/Engine/Source/Editor/SequenceRecorder/Private/SequenceRecorderPrivatePCH.h
@@ -12,4 +12,6 @@
 #include "KeyParams.h"
 #include "IKeyframeSection.h"
 #include "MovieSceneSection.h"
-#include "LevelSequence.h"
\ No newline at end of file
+#include "LevelSequence.h"
+#include "IMovieSceneSectionRecorder.h"
+#include "IMovieSceneSectionRecorderFactory.h"
\ No newline at end of file
diff --git a/Engine/Source/Editor/SequenceRecorder/Private/SequenceRecorderSettings.cpp b/Engine/Source/Editor/SequenceRecorder/Private/SequenceRecorderSettings.cpp
--- a/Engine/Source/Editor/SequenceRecorder/Private/SequenceRecorderSettings.cpp
+++ b/Engine/Source/Editor/SequenceRecorder/Private/SequenceRecorderSettings.cpp
@@ -15,6 +15,11 @@ USequenceRecorderSettings::USequenceRecorderSettings(const FObjectInitializer& O
 	bRecordNearbySpawnedActors = true;
 	NearbyActorRecordingProximity = 5000.0f;
 	bRecordWorldSettingsActor = true;
+
+	ComponentClassesToRecord.Add(USkeletalMeshComponent::StaticClass());
+	ComponentClassesToRecord.Add(UStaticMeshComponent::StaticClass());
+	ComponentClassesToRecord.Add(UParticleSystemComponent::StaticClass());
+	ComponentClassesToRecord.Add(ULightComponent::StaticClass());
 }
 
 void USequenceRecorderSettings::PostEditChangeChainProperty(struct FPropertyChangedChainEvent& PropertyChangedEvent)
diff --git a/Engine/Source/Editor/SequenceRecorder/Public/ISequenceRecorder.h b/Engine/Source/Editor/SequenceRecorder/Public/ISequenceRecorder.h
--- a/Engine/Source/Editor/SequenceRecorder/Public/ISequenceRecorder.h
+++ b/Engine/Source/Editor/SequenceRecorder/Public/ISequenceRecorder.h
@@ -53,4 +53,11 @@ public:
 	 * @param	Actor	The actor that was 'de-spawned'
 	 */
 	virtual void NotifyActorStopRecording(AActor* Actor) = 0;
+
+	/**
+	 * Get the spawnable Guid in the currently recording movie scene for the specified actor.
+	 * @param	Actor	The Actor to check
+	 * @return the spawnable Guid
+	 */
+	virtual FGuid GetRecordingGuid(AActor* Actor) const = 0;
 };
\ No newline at end of file
diff --git a/Engine/Source/Editor/SequenceRecorder/Public/SequenceRecorderSettings.h b/Engine/Source/Editor/SequenceRecorder/Public/SequenceRecorderSettings.h
--- a/Engine/Source/Editor/SequenceRecorder/Public/SequenceRecorderSettings.h
+++ b/Engine/Source/Editor/SequenceRecorder/Public/SequenceRecorderSettings.h
@@ -61,4 +61,8 @@ public:
 	/** Default settings applied to animation recording */
 	UPROPERTY(Config, EditAnywhere, Category = "Sequence Recording")
 	FAnimationRecordingSettings DefaultAnimationSettings;
+
+	/** Component classes we record by default. If an actor does not contain one of these classes it will be ignored. */
+	UPROPERTY(Config, EditAnywhere, AdvancedDisplay, Category = "Sequence Recording")
+	TArray<TSubclassOf<USceneComponent>> ComponentClassesToRecord;
 };
\ No newline at end of file
diff --git a/Engine/Source/Editor/SequenceRecorder/SequenceRecorder.Build.cs b/Engine/Source/Editor/SequenceRecorder/SequenceRecorder.Build.cs
--- a/Engine/Source/Editor/SequenceRecorder/SequenceRecorder.Build.cs
+++ b/Engine/Source/Editor/SequenceRecorder/SequenceRecorder.Build.cs
@@ -14,7 +14,7 @@ namespace UnrealBuildTool.Rules
 			PrivateIncludePaths.AddRange(
 				new string[] {
 					"Editor/SequenceRecorder/Private",
-					"Editor/SequenceRecorder/Private/PropertyRecorders",
+					"Editor/SequenceRecorder/Private/Sections",
 				}
 				);
 
diff --git a/Engine/Source/Editor/SequenceRecorderSections/Private/MovieSceneSpawnSectionRecorder.cpp b/Engine/Source/Editor/SequenceRecorderSections/Private/MovieSceneSpawnSectionRecorder.cpp
new file mode 100644
--- /dev/null
+++ b/Engine/Source/Editor/SequenceRecorderSections/Private/MovieSceneSpawnSectionRecorder.cpp
@@ -0,0 +1,63 @@
+// Copyright 1998-2015 Epic Games, Inc. All Rights Reserved.
+
+#include "SequenceRecorderSectionsPrivatePCH.h"
+#include "MovieSceneSpawnSectionRecorder.h"
+#include "MovieSceneBoolSection.h"
+#include "MovieSceneSpawnTrack.h"
+#include "MovieScene.h"
+
+TSharedPtr<IMovieSceneSectionRecorder> FMovieSceneSpawnSectionRecorderFactory::CreateSectionRecorder(const struct FActorRecordingSettings& InActorRecordingSettings) const
+{
+	return MakeShareable(new FMovieSceneSpawnSectionRecorder);
+}
+
+bool FMovieSceneSpawnSectionRecorderFactory::CanRecordObject(UObject* InObjectToRecord) const
+{
+	return InObjectToRecord->IsA<AActor>();
+}
+
+void FMovieSceneSpawnSectionRecorder::CreateSection(UObject* InObjectToRecord, UMovieScene* MovieScene, const FGuid& Guid, float Time)
+{
+	ObjectToRecord = InObjectToRecord;
+
+	UMovieSceneSpawnTrack* SpawnTrack = MovieScene->AddTrack<UMovieSceneSpawnTrack>(Guid);
+	if(SpawnTrack)
+	{
+		MovieSceneSection = Cast<UMovieSceneBoolSection>(SpawnTrack->CreateNewSection());
+
+		SpawnTrack->AddSection(*MovieSceneSection);
+		SpawnTrack->SetObjectId(Guid);
+
+		MovieSceneSection->SetDefault(false);
+		MovieSceneSection->AddKey(0.0f, false, EMovieSceneKeyInterpolation::Break);
+
+		MovieSceneSection->SetStartTime(Time);
+		MovieSceneSection->SetIsInfinite(true);
+	}
+
+	bWasSpawned = false;
+}
+
+void FMovieSceneSpawnSectionRecorder::FinalizeSection()
+{
+	const bool bSpawned = ObjectToRecord.IsValid();
+	if(bSpawned != bWasSpawned)
+	{
+		MovieSceneSection->AddKey(MovieSceneSection->GetEndTime(), bSpawned, EMovieSceneKeyInterpolation::Break);
+	}
+}
+
+void FMovieSceneSpawnSectionRecorder::Record(float CurrentTime)
+{
+	if(ObjectToRecord.IsValid())
+	{
+		MovieSceneSection->SetEndTime(CurrentTime);
+	}
+
+	const bool bSpawned = ObjectToRecord.IsValid();
+	if(bSpawned != bWasSpawned)
+	{
+		MovieSceneSection->AddKey(CurrentTime, bSpawned, EMovieSceneKeyInterpolation::Break);
+	}
+	bWasSpawned = bSpawned;
+}
\ No newline at end of file
diff --git a/Engine/Source/Editor/SequenceRecorderSections/Private/MovieSceneVisibilitySectionRecorder.cpp b/Engine/Source/Editor/SequenceRecorderSections/Private/MovieSceneVisibilitySectionRecorder.cpp
new file mode 100644
--- /dev/null
+++ b/Engine/Source/Editor/SequenceRecorderSections/Private/MovieSceneVisibilitySectionRecorder.cpp
@@ -0,0 +1,86 @@
+// Copyright 1998-2015 Epic Games, Inc. All Rights Reserved.
+
+#include "SequenceRecorderSectionsPrivatePCH.h"
+#include "MovieSceneVisibilitySectionRecorder.h"
+#include "MovieSceneVisibilitySection.h"
+#include "MovieSceneVisibilityTrack.h"
+#include "MovieScene.h"
+#include "ActorRecordingSettings.h"
+
+static const FName VisibilityTrackName = TEXT("Visibility");
+
+TSharedPtr<IMovieSceneSectionRecorder> FMovieSceneVisibilitySectionRecorderFactory::CreateSectionRecorder(const struct FActorRecordingSettings& InActorRecordingSettings) const
+{
+	UMovieSceneVisibilitySectionRecorderSettings* Settings = InActorRecordingSettings.GetSettingsObject<UMovieSceneVisibilitySectionRecorderSettings>();
+	check(Settings);
+	if (Settings->bRecordVisibility)
+	{
+		return MakeShareable(new FMovieSceneVisibilitySectionRecorder);
+	}
+
+	return nullptr;
+}
+
+bool FMovieSceneVisibilitySectionRecorderFactory::CanRecordObject(UObject* InObjectToRecord) const
+{
+	return InObjectToRecord->IsA<AActor>() || InObjectToRecord->IsA<USceneComponent>();
+}
+
+void FMovieSceneVisibilitySectionRecorder::CreateSection(UObject* InObjectToRecord, UMovieScene* MovieScene, const FGuid& Guid, float Time)
+{
+	ObjectToRecord = InObjectToRecord;
+
+	UMovieSceneVisibilityTrack* VisibilityTrack = MovieScene->AddTrack<UMovieSceneVisibilityTrack>(Guid);
+	if(VisibilityTrack)
+	{
+		VisibilityTrack->SetPropertyNameAndPath(VisibilityTrackName, VisibilityTrackName.ToString());
+
+		MovieSceneSection = Cast<UMovieSceneVisibilitySection>(VisibilityTrack->CreateNewSection());
+
+		VisibilityTrack->AddSection(*MovieSceneSection);
+
+		MovieSceneSection->SetDefault(false);
+
+		bWasVisible = false;
+		if(USceneComponent* SceneComponent = Cast<USceneComponent>(ObjectToRecord.Get()))
+		{
+			bWasVisible = SceneComponent->IsVisible() && SceneComponent->IsRegistered();
+		}
+		else if(AActor* Actor = Cast<AActor>(ObjectToRecord.Get()))
+		{
+			bWasVisible = !Actor->bHidden;
+		}
+
+		MovieSceneSection->AddKey(Time, bWasVisible, EMovieSceneKeyInterpolation::Break);
+
+		MovieSceneSection->SetStartTime(Time);
+	}
+}
+
+void FMovieSceneVisibilitySectionRecorder::FinalizeSection()
+{
+}
+
+void FMovieSceneVisibilitySectionRecorder::Record(float CurrentTime)
+{
+	if(ObjectToRecord.IsValid())
+	{
+		MovieSceneSection->SetEndTime(CurrentTime);
+
+		bool bVisible = false;
+		if(USceneComponent* SceneComponent = Cast<USceneComponent>(ObjectToRecord.Get()))
+		{
+			bVisible = SceneComponent->IsVisible() && SceneComponent->IsRegistered();
+		}
+		else if(AActor* Actor = Cast<AActor>(ObjectToRecord.Get()))
+		{
+			bVisible = !Actor->bHidden;
+		}
+
+		if(bVisible != bWasVisible)
+		{
+			MovieSceneSection->AddKey(CurrentTime, bVisible, EMovieSceneKeyInterpolation::Break);
+		}
+		bWasVisible = bVisible;
+	}
+}
\ No newline at end of file
diff --git a/Engine/Source/Editor/SequenceRecorderSections/Private/SequenceRecorderSectionsPrivatePCH.h b/Engine/Source/Editor/SequenceRecorderSections/Private/SequenceRecorderSectionsPrivatePCH.h
new file mode 100644
--- /dev/null
+++ b/Engine/Source/Editor/SequenceRecorderSections/Private/SequenceRecorderSectionsPrivatePCH.h
@@ -0,0 +1,7 @@
+// Copyright 1998-2016 Epic Games, Inc. All Rights Reserved.
+
+#include "Engine.h"
+#include "MovieScene.h"
+#include "KeyParams.h"
+#include "IMovieSceneSectionRecorderFactory.h"
+#include "IMovieSceneSectionRecorder.h"
\ No newline at end of file
diff --git a/Engine/Source/Editor/SequenceRecorderSections/SequenceRecorderSections.Build.cs b/Engine/Source/Editor/SequenceRecorderSections/SequenceRecorderSections.Build.cs
new file mode 100644
--- /dev/null
+++ b/Engine/Source/Editor/SequenceRecorderSections/SequenceRecorderSections.Build.cs
@@ -0,0 +1,45 @@
+// Copyright 1998-2016 Epic Games, Inc. All Rights Reserved.
+
+namespace UnrealBuildTool.Rules
+{
+	public class SequenceRecorderSections : ModuleRules
+	{
+		public SequenceRecorderSections(TargetInfo Target)
+		{
+			PublicIncludePaths.AddRange(
+				new string[] {
+				}
+				);
+
+			PrivateIncludePaths.AddRange(
+				new string[] {
+					"Editor/SequenceRecorderSections/Private",
+				}
+				);
+
+			PublicDependencyModuleNames.AddRange(
+				new string[]
+				{
+				}
+				);
+
+			PrivateDependencyModuleNames.AddRange(
+				new string[]
+				{
+					"Core",
+					"CoreUObject",
+                    "Engine",
+                    "MovieScene",
+                    "MovieSceneTracks",
+					"SequenceRecorder",
+				}
+				);
+
+			DynamicallyLoadedModuleNames.AddRange(
+				new string[]
+				{
+				}
+				);
+		}
+	}
+}
\ No newline at end of file
diff --git a/Engine/Source/Editor/Sequencer/Private/DisplayNodes/SequencerObjectBindingNode.cpp b/Engine/Source/Editor/Sequencer/Private/DisplayNodes/SequencerObjectBindingNode.cpp
--- a/Engine/Source/Editor/Sequencer/Private/DisplayNodes/SequencerObjectBindingNode.cpp
+++ b/Engine/Source/Editor/Sequencer/Private/DisplayNodes/SequencerObjectBindingNode.cpp
@@ -96,27 +96,11 @@ void FSequencerObjectBindingNode::BuildContextMenu(FMenuBuilder& MenuBuilder)
 
 		if (Spawnable)
 		{
-			UBlueprint* Blueprint = Cast<UBlueprint>(Spawnable->GetClass()->ClassGeneratedBy);
-
-			if (Blueprint)
-			{
-				FFormatNamedArguments Args;
-				MenuBuilder.AddMenuEntry(
-					FText::Format( LOCTEXT("EditDefaults", "Edit Defaults"), Args),
-					FText::Format( LOCTEXT("EditDefaultsTooltip", "Edit the defaults for this spawnable object"), Args ),
-					FSlateIcon(),
-					FUIAction(FExecuteAction::CreateLambda([=]{
-						FAssetEditorManager::Get().OpenEditorForAsset(Blueprint);
-					}))
-				);
-			}
-
 			MenuBuilder.AddSubMenu(
 				LOCTEXT("OwnerLabel", "Spawned Object Owner"),
 				LOCTEXT("OwnerTooltip", "Specifies how the spawned object is to be owned"),
 				FNewMenuDelegate::CreateSP(this, &FSequencerObjectBindingNode::AddSpawnOwnershipMenu)
 			);
-
 		}
 		else
 		{
@@ -410,48 +394,12 @@ const UClass* FSequencerObjectBindingNode::GetClassForObjectBinding() const
 	
 	// should exist, but also shouldn't be both a spawnable and a possessable
 	check((Spawnable != nullptr) ^ (Possessable != nullptr));
-	check((nullptr == Spawnable) || (nullptr != Spawnable->GetClass())  );
-	const UClass* ObjectClass = Spawnable ? Spawnable->GetClass()->GetSuperClass() : Possessable->GetPossessedObjectClass();
+	check((nullptr == Spawnable) || (nullptr != Spawnable->GetObjectTemplate())  );
+	const UClass* ObjectClass = Spawnable ? Spawnable->GetObjectTemplate()->GetClass() : Possessable->GetPossessedObjectClass();
 
 	return ObjectClass;
 }
 
-UObject* FSequencerObjectBindingNode::FindRepresentativeObject()
-{
-	FSequencer& Sequencer = GetSequencer();
-	
-	UObject* BoundObject = Sequencer.GetFocusedMovieSceneSequenceInstance()->FindObject(ObjectBinding, Sequencer);
-	if (BoundObject)
-	{
-		return BoundObject;
-	}
-
-	UMovieScene* FocusedMovieScene = Sequencer.GetFocusedMovieSceneSequence()->GetMovieScene();
-
-	FMovieScenePossessable* Possessable = FocusedMovieScene->FindPossessable(ObjectBinding);
-	// If we're a possessable with a parent spawnable and we don't have the object, we look the object up within the default object of the spawnable
-	if (Possessable && Possessable->GetParent().IsValid())
-	{
-		// If we're a spawnable and we don't have the object, use the default object to build up the track menu
-		FMovieSceneSpawnable* ParentSpawnable = FocusedMovieScene->FindSpawnable(Possessable->GetParent());
-		if (ParentSpawnable)
-		{
-			UObject* ParentObject = ParentSpawnable->GetClass()->GetDefaultObject();
-			if (ParentObject)
-			{
-				return Sequencer.GetFocusedMovieSceneSequence()->FindPossessableObject(ObjectBinding, ParentObject);
-			}
-		}
-	}
-	// If we're a spawnable and we don't have the object, use the default object to build up the track menu
-	else if (FMovieSceneSpawnable* Spawnable = FocusedMovieScene->FindSpawnable(ObjectBinding))
-	{
-		return Spawnable->GetClass()->GetDefaultObject();
-	}
-
-	return nullptr;
-}
-
 /* FSequencerObjectBindingNode callbacks
  *****************************************************************************/
 
@@ -462,7 +410,7 @@ TSharedRef<SWidget> FSequencerObjectBindingNode::HandleAddTrackComboButtonGetMen
 	//@todo need to resolve this between UMG and the level editor sequencer
 	const bool bUseSubMenus = Sequencer.IsLevelEditorSequencer();
 
-	UObject* BoundObject = FindRepresentativeObject();
+	UObject* BoundObject = GetSequencer().FindSpawnedObjectOrTemplate(ObjectBinding);
 
 	ISequencerModule& SequencerModule = FModuleManager::GetModuleChecked<ISequencerModule>( "Sequencer" );
 	TSharedRef<FUICommandList> CommandList(new FUICommandList);
@@ -630,7 +578,7 @@ void FSequencerObjectBindingNode::HandleLabelsSubMenuCreate(FMenuBuilder& MenuBu
 void FSequencerObjectBindingNode::HandlePropertyMenuItemExecute(TArray<UProperty*> PropertyPath)
 {
 	FSequencer& Sequencer = GetSequencer();
-	UObject* BoundObject = FindRepresentativeObject();
+	UObject* BoundObject = GetSequencer().FindSpawnedObjectOrTemplate(ObjectBinding);
 
 	TArray<UObject*> KeyableBoundObjects;
 	if (BoundObject != nullptr)
diff --git a/Engine/Source/Editor/Sequencer/Private/DisplayNodes/SequencerObjectBindingNode.h b/Engine/Source/Editor/Sequencer/Private/DisplayNodes/SequencerObjectBindingNode.h
--- a/Engine/Source/Editor/Sequencer/Private/DisplayNodes/SequencerObjectBindingNode.h
+++ b/Engine/Source/Editor/Sequencer/Private/DisplayNodes/SequencerObjectBindingNode.h
@@ -72,8 +72,6 @@ protected:
 	const UClass* GetClassForObjectBinding() const;
 
 private:
-
-	UObject* FindRepresentativeObject();
 	
 	TSharedRef<SWidget> HandleAddTrackComboButtonGetMenuContent();
 	
diff --git a/Engine/Source/Editor/Sequencer/Private/SAnimationOutlinerTreeNode.cpp b/Engine/Source/Editor/Sequencer/Private/SAnimationOutlinerTreeNode.cpp
--- a/Engine/Source/Editor/Sequencer/Private/SAnimationOutlinerTreeNode.cpp
+++ b/Engine/Source/Editor/Sequencer/Private/SAnimationOutlinerTreeNode.cpp
@@ -86,9 +86,10 @@ SAnimationOutlinerTreeNode::~SAnimationOutlinerTreeNode()
 void SAnimationOutlinerTreeNode::Construct( const FArguments& InArgs, TSharedRef<FSequencerDisplayNode> Node, const TSharedRef<SSequencerTreeViewRow>& InTableRow )
 {
 	DisplayNode = Node;
-	bIsTopLevelNode = !Node->GetParent().IsValid();
+	bIsOuterTopLevelNode = !Node->GetParent().IsValid();
+	bIsInnerTopLevelNode = Node->GetType() != ESequencerNode::Folder && Node->GetParent().IsValid() && Node->GetParent()->GetType() == ESequencerNode::Folder;
 
-	if (bIsTopLevelNode)
+	if (bIsOuterTopLevelNode)
 	{
 		ExpandedBackgroundBrush = FEditorStyle::GetBrush( "Sequencer.AnimationOutliner.TopLevelBorder_Expanded" );
 		CollapsedBackgroundBrush = FEditorStyle::GetBrush( "Sequencer.AnimationOutliner.TopLevelBorder_Collapsed" );
@@ -99,6 +100,18 @@ void SAnimationOutlinerTreeNode::Construct( const FArguments& InArgs, TSharedRef
 		CollapsedBackgroundBrush = FEditorStyle::GetBrush( "Sequencer.AnimationOutliner.DefaultBorder" );
 	}
 
+	FMargin InnerNodePadding;
+	if ( bIsInnerTopLevelNode )
+	{
+		InnerBackgroundBrush = FEditorStyle::GetBrush( "Sequencer.AnimationOutliner.TopLevelBorder_Expanded" );
+		InnerNodePadding = FMargin(0.f, 1.f);
+	}
+	else
+	{
+		InnerBackgroundBrush = FEditorStyle::GetBrush( "Sequencer.AnimationOutliner.TransparentBorder" );
+		InnerNodePadding = FMargin(0.f);
+	}
+
 	TableRowStyle = &FEditorStyle::Get().GetWidgetStyle<FTableRowStyle>("TableView.Row");
 
 	FSlateFontInfo NodeFont = FEditorStyle::GetFontStyle("Sequencer.AnimationOutliner.RegularFont");
@@ -137,56 +150,68 @@ void SAnimationOutlinerTreeNode::Construct( const FArguments& InArgs, TSharedRef
 
 					// Expand track lanes button
 					+ SHorizontalBox::Slot()
-					.Padding(FMargin(2.f, 0.f, 4.f, 0.f))
+					.Padding(FMargin(2.f, 0.f, 2.f, 0.f))
 					.VAlign( VAlign_Center )
 					.AutoWidth()
 					[
 						SNew(SExpanderArrow, InTableRow).IndentAmount(SequencerLayoutConstants::IndentAmount)
 					]
 
-					// Icon
 					+ SHorizontalBox::Slot()
-					.Padding(FMargin(0.f, 0.f, 4.f, 0.f))
-					.VAlign(VAlign_Center)
-					.AutoWidth()
+					.Padding( InnerNodePadding )
 					[
-						SNew(SOverlay)
-
-						+ SOverlay::Slot()
-						[
-							SNew(SImage)
-							.Image(InArgs._IconBrush)
-						]
-
-						+ SOverlay::Slot()
-						.VAlign(VAlign_Top)
-						.HAlign(HAlign_Right)
+						SNew( SBorder )
+						.BorderImage( FEditorStyle::GetBrush( "Sequencer.AnimationOutliner.TopLevelBorder_Collapsed" ) )
+						.BorderBackgroundColor( this, &SAnimationOutlinerTreeNode::GetNodeInnerBackgroundTint )
+						.Padding( FMargin(0) )
 						[
-							SNew(SImage)
-							.Image(InArgs._IconOverlayBrush)
+							SNew( SHorizontalBox )
+							
+							// Icon
+							+ SHorizontalBox::Slot()
+							.Padding(FMargin(0.f, 0.f, 4.f, 0.f))
+							.VAlign(VAlign_Center)
+							.AutoWidth()
+							[
+								SNew(SOverlay)
+
+								+ SOverlay::Slot()
+								[
+									SNew(SImage)
+									.Image(InArgs._IconBrush)
+								]
+
+								+ SOverlay::Slot()
+								.VAlign(VAlign_Top)
+								.HAlign(HAlign_Right)
+								[
+									SNew(SImage)
+									.Image(InArgs._IconOverlayBrush)
+								]
+
+								+ SOverlay::Slot()
+								[
+									SNew(SSpacer)
+									.Visibility(EVisibility::Visible)
+									.ToolTipText(InArgs._IconToolTipText)
+								]
+							]
+
+							// Label Slot
+							+ SHorizontalBox::Slot()
+							.VAlign(VAlign_Center)
+							.Padding(FMargin(0.f, 0.f, 4.f, 0.f))
+							.AutoWidth()
+							[
+								EditableLabel.ToSharedRef()
+							]
+
+							// Arbitrary customization slot
+							+ SHorizontalBox::Slot()
+							[
+								InArgs._CustomContent.Widget
+							]
 						]
-
-						+ SOverlay::Slot()
-						[
-							SNew(SSpacer)
-							.Visibility(EVisibility::Visible)
-							.ToolTipText(InArgs._IconToolTipText)
-						]
-					]
-
-					// Label Slot
-					+ SHorizontalBox::Slot()
-					.VAlign(VAlign_Center)
-					.Padding(FMargin(0.f, 0.f, 4.f, 0.f))
-					.AutoWidth()
-					[
-						EditableLabel.ToSharedRef()
-					]
-
-					// Arbitrary customization slot
-					+ SHorizontalBox::Slot()
-					[
-						InArgs._CustomContent.Widget
 					]
 				]
 			]
@@ -286,11 +311,41 @@ FSlateColor SAnimationOutlinerTreeNode::GetNodeBackgroundTint() const
 	}	
 	else if (DisplayNode->IsHovered())
 	{
-		return bIsTopLevelNode ? FLinearColor(FColor(52, 52, 52, 255)) : FLinearColor(FColor(72, 72, 72, 255));
+		return bIsOuterTopLevelNode ? FLinearColor(FColor(52, 52, 52, 255)) : FLinearColor(FColor(72, 72, 72, 255));
+	}
+	else
+	{
+		return bIsOuterTopLevelNode ? FLinearColor(FColor(48, 48, 48, 255)) : FLinearColor(FColor(62, 62, 62, 255));
+	}
+}
+
+FSlateColor SAnimationOutlinerTreeNode::GetNodeInnerBackgroundTint() const
+{
+	if ( bIsInnerTopLevelNode )
+	{
+		FSequencer& Sequencer = DisplayNode->GetSequencer();
+		const bool bIsSelected = Sequencer.GetSelection().IsSelected( DisplayNode.ToSharedRef() );
+
+		if ( bIsSelected )
+		{
+			return FEditorStyle::GetSlateColor( "SelectionColor_Pressed" );
+		}
+		else if ( Sequencer.GetSelection().NodeHasSelectedKeysOrSections( DisplayNode.ToSharedRef() ) )
+		{
+			return FLinearColor( FColor( 115, 115, 115, 255 ) );
+		}
+		else if ( DisplayNode->IsHovered() )
+		{
+			return FLinearColor( FColor( 52, 52, 52, 255 ) );
+		}
+		else
+		{
+			return FLinearColor( FColor( 48, 48, 48, 255 ) );
+		}
 	}
 	else
 	{
-		return bIsTopLevelNode ? FLinearColor(FColor(48, 48, 48, 255)) : FLinearColor(FColor(62, 62, 62, 255));
+		return FLinearColor( 0.f, 0.f, 0.f, 0.f );
 	}
 }
 
diff --git a/Engine/Source/Editor/Sequencer/Private/SAnimationOutlinerTreeNode.h b/Engine/Source/Editor/Sequencer/Private/SAnimationOutlinerTreeNode.h
--- a/Engine/Source/Editor/Sequencer/Private/SAnimationOutlinerTreeNode.h
+++ b/Engine/Source/Editor/Sequencer/Private/SAnimationOutlinerTreeNode.h
@@ -63,6 +63,11 @@ private:
 	 */
 	FSlateColor GetNodeBackgroundTint() const;
 
+	/**
+	* @return The tint to apply to the border image for the inner portion of the node.
+	*/
+	FSlateColor GetNodeInnerBackgroundTint() const;
+
 	/**
 	 * @return The expander visibility of this node.
 	 */
@@ -103,8 +108,11 @@ private:
 	/** Holds the editable text label widget. */
 	TSharedPtr<SEditableLabel> EditableLabel;
 
-	/** True if this node is a top level node, false otherwise */
-	bool bIsTopLevelNode;
+	/** True if this node is a top level node, at the root of the tree, false otherwise */
+	bool bIsOuterTopLevelNode;
+
+	/** True if this is a top level node inside or a folder, otherwise false. */
+	bool bIsInnerTopLevelNode;
 
 	/** Default background brush for this node when expanded */
 	const FSlateBrush* ExpandedBackgroundBrush;
@@ -112,6 +120,9 @@ private:
 	/** Default background brush for this node when collapsed */
 	const FSlateBrush* CollapsedBackgroundBrush;
 
+	/** The brush to use when drawing the background for the inner portion of the node. */
+	const FSlateBrush* InnerBackgroundBrush;
+
 	/** The table row style used for nodes in the tree. This is required as we don't actually use the tree for selection. */
 	const FTableRowStyle* TableRowStyle;
 };
\ No newline at end of file
