Copying //UE4/Dev-Sequencer to //UE4/Dev-Main (Source: //UE4/Dev-Sequencer @ 3048536)

==========================
MAJOR FEATURES + CHANGES
==========================

Change 3004054 on 2016/06/07 by [REDACTED_PERSON]

	Movie Capture: Expose compressed EXR frames to Sequencer Capture UI and command line.

Change 3007048 on 2016/06/09 by [REDACTED_PERSON]

	Sequencer: Allow showing the context menu for more than one selected node. Certain operations (ie. Lock, Active, Label) now operate on all the selected nodes.

Change 3007795 on 2016/06/09 by [REDACTED_PERSON]

	Sequencer: Select actors for corresponding selected keys or sections.

Change 3009689 on 2016/06/10 by [REDACTED_PERSON]

	Movie Capture: Add {shot} description to movie capture filename format.

Change 3010180 on 2016/06/11 by [REDACTED_PERSON]

	Sequencer: Add support for copying and converting linear color tracks from matinee.

Change 3012472 on 2016/06/14 by [REDACTED_PERSON]

	Sequence Recorder: Optimize adding keys to rich curves. Set the times and values at once as an ordered set.

Change 3012473 on 2016/06/14 by [REDACTED_PERSON]

	Sequence Recorder: Find corresponding PIE world actor for Actor to Record. This fixes a bug where if you set the actor before PIE and then record, the recording doesn't work.

Change 3012474 on 2016/06/14 by [REDACTED_PERSON]

	Sequence Recorder: Fix countdown timer so that it draws in Simulate mode as well as Play.

Change 3014868 on 2016/06/15 by [REDACTED_PERSON]

	Sequencer: For legacy, properties with the CPF_Interp flag can be animated in Matinee. It doesn't require the CPF_Edit flag as well. This makes Sequencer consistent with Matinee.

Change 3014869 on 2016/06/15 by [REDACTED_PERSON]

	Sequencer: Rename "Expose to Matinee" to "Expose to Cinematics"

Change 3016137 on 2016/06/16 by [REDACTED_PERSON]

	Sequencer: Added support for the named "PerformanceCapture" event which like Matinee, calls GEngine->PerformanceCapture to output a screenshot when the event fires. Refactor event track/sections so that the player is passed to the trigger events evaluation.

Change 3018996 on 2016/06/18 by [REDACTED_PERSON]

	Sequencer: Add play/pause/stop events to level sequence player ala matinee actor.

Change 3019763 on 2016/06/20 by [REDACTED_PERSON]

	Sequencer: Fix key editor commit when user tabs away (lose focus). Also, make the key navigation buttons and track color button not focusable.

	Slate: Make editable labels not focusable if they are not editable.

Change 3019768 on 2016/06/20 by [REDACTED_PERSON]

	Sequencer: Event track no longer fires if the playback status is stopped. This fixes a bug where when playback is stopped and the movie scene sequence is returned to the start of playback, we don't want all the events from the last playback position to the start of playback to fire.

Change 3020849 on 2016/06/21 by [REDACTED_PERSON]

	Sequencer: Fixed blueprint classes as spawnables not being reinstanced correctly

	  - The RF_ArchetypeObject flag was previously used to denote spawnable object templates, however this caused blueprint reinstancing to skip such objects, which resulted in crashes, or data loss
	  - Added code to rename expired object templates to ensure there is no name collision
	  - Ensured that register functions are only called on actors that are part of a world (when finalizing blueprint reinstancing)

Change 3021400 on 2016/06/21 by [REDACTED_PERSON]

	Sequencer - Add support for exporting to fbx.

Change 3022941 on 2016/06/22 by [REDACTED_PERSON]

	Sequencer: Thumbnail improvements

	  - Fixed poor performance when continually zooming (thumbnails were being constantly and immediately rendered in this case)
	  - Added quality setting to thumbnail sections (draft/normal/best)
	  - Improved fade transition to use proper alpha blending
	  - Fixed needlessly recreating the entire world's render state when creating a new editor viewport client. This caused a significant hitch opening sequences in certain large worlds.

Change 3022944 on 2016/06/22 by [REDACTED_PERSON]

	Sequencer: Fixed jitter when jumping around

	  - bEditorCameraCut was not being reset on the next frame, which prevented temporal effects from running

Change 3024774 on 2016/06/23 by [REDACTED_PERSON]

	Sequencer: Cache show intensity to invalidate the waveform preview when the intensity changes.

Change 3026170 on 2016/06/24 by [REDACTED_PERSON]

	Sequence Recorder: Change Actor Name to Record to a TLazyObjectPtr<AActor> so that the user can choose the actor directly rather than typing in a name.

Change 3026181 on 2016/06/24 by [REDACTED_PERSON]

	Sequencer: Add visibility options to show/hide/only when selected 3d trajectories per transform section.

Change 3026489 on 2016/06/24 by [REDACTED_PERSON]

	Sequencer: Fixed some recorded actors not being saved into level sequences correctly
	  - The use of StaticDuplicateObject was causing temporary/transient or otherwise external data to be duplicated into template actors which caused the actor to be culled on package save.
	  - Using NewObject followed by copying the object properties guarantees we don't copy this data.
	  - This is the same method used to spawn the spawnable actor

Change 3026522 on 2016/06/24 by [REDACTED_PERSON]

	Sequence Recorder: Fix bug where sample rate wouldn't get set if the length is set to 0.

Change 3027768 on 2016/06/25 by [REDACTED_PERSON]

	Sequencer: Fix player stopping after 60 seconds. Assign the player to a UPROPERTY so that it's not garbage collected.

Change 3028318 on 2016/06/27 by [REDACTED_PERSON]

	Editor: Added safety check to prevent stack overflow populating the world outliner

	  - It was possible for the population code to get stuck in an infinite loop if an actor happened to be attached to itself.

Change 3034262 on 2016/06/30 by [REDACTED_PERSON]

	Sequencer: Fixed invalid properties being recorded when creating spawnable object templates

	 - When creating a spawnable out of an attached actor, the attachment was getting copied into the template. This results in GLEO errors when saving the package.
	 - Suppressed object replacement notifications when calling CopyPropertiesForUnrelatedObjects for spawnables

Change 3035168 on 2016/06/30 by [REDACTED_PERSON]

	UMG: Set anim range to the playback range size. This fixes a bug where animation doesn't play the full length of the playback range if it starts negative.

Change 3035169 on 2016/06/30 by [REDACTED_PERSON]

	Sequencer: Add ReversePlay() and ChangePlaybackDirection() functions (ala Matinee).

Change 3035174 on 2016/06/30 by [REDACTED_PERSON]

	Sequencer: Fix evaluation when playback starts or loops around so that last time is enforced to be the lower bound of the playback range. This fixes a bug in the particle track where if there's a particle that triggers at time 0 and a level sequence that starts at time 0, the event will be missed since last time will be 0 but then rounded to a fixed frame with epislon .0001f.

Change 3035186 on 2016/06/30 by [REDACTED_PERSON]

	Sequencer: Add reset selection range and remove unused delete selection range.

Change 3035197 on 2016/07/01 by [REDACTED_PERSON]

	Sequencer: Fix so that adding a sub section adds to the clicked on sub track and not just the first.

Change 3036586 on 2016/07/02 by [REDACTED_PERSON]

	Sequencer: More play controls - shuttle backward, pause, shuttle forward (j, k, l)

Change 3036941 on 2016/07/04 by [REDACTED_PERSON]

	Sequencer: Record transforms in world space where an actor is attached, and we're not recording its parent

Change 3039290 on 2016/07/06 by [REDACTED_PERSON]

	Sequencer: Various capture fixes
	  - Movie captures no longer crash when no world is loaded (they gracefully close instead)
	    - Currently waiting on change from core to hook up the error code with an actual process termination code
	  - We now force -NoLoadingScreen, -Windowed and -ForceRes since movie captures will not work without these

Change 3039831 on 2016/07/06 by [REDACTED_PERSON]

	Sequencer - Notify data changed refactor.
	+ Add a change type to "NotifyMovieSceneDataChanged" so that sequencer knows what parts of the system to refresh.
	+ Remove most calls for UpdateRuntimeInstances and replace them with a call to NotifyMovieSceneDataChanged.
	+ Update UMG so that it copies the animation data to the compiled class whenever it's changed.

Change 3044087 on 2016/07/10 by [REDACTED_PERSON]

	Sequencer - Prevent crashes when encountering filler shots with no valid sequence.

Change 3044151 on 2016/07/10 by [REDACTED_PERSON]

	Sequencer: Only update selected nodes if they change. This fixes a bug in the curve editor where undo reselects and autoframes.

Change 3044164 on 2016/07/10 by [REDACTED_PERSON]

	Sequencer: Added ability to immediately record actors directly into sequencer
	  - "Record 'ActorName' In Sequencer" option is now available on the level editor context menu for selected actors when sequencer is open.
	  - This immediately triggers a countdown and records the currently selected actors into a sub sequence in the currently focussed movie scene
	  - Creates a cinematic shot track if you record a camera
	  - Removed the older "queue" and "trigger" methods for now to make it cleaner for the demo.

Change 3044180 on 2016/07/11 by [REDACTED_PERSON]

	Sequencer: Added ability to possess viewports while in PIE
	  - Added a new option to the level sequence editor settings to allow possession of PIE viewports

Change 3044181 on 2016/07/11 by [REDACTED_PERSON]

	Sequencer: Added ability to specify event contexts for FSequencer
	  - This allows us to trigger events from playback within sequencer, according to the sequencer client

Change 3044188 on 2016/07/11 by [REDACTED_PERSON]

	Sequencer: We no longer evaluate camera cut tracks as part of sub tracks, only shot tracks.

Change 3044193 on 2016/07/11 by [REDACTED_PERSON]

	Sequencer: Added cvar LevelSequence.DefaultFixedFrameIntervalPlayback to control this setting for newly created level sequences

Change 3044194 on 2016/07/11 by [REDACTED_PERSON]

	Sequencer: Added an option to rewind the sequence when a recording is started
	Defaults to 'on'

Change 3047334 on 2016/07/12 by [REDACTED_PERSON]

	Sequencer: Add transactions for creating a camera cut track and a folder.

Change 3047365 on 2016/07/12 by [REDACTED_PERSON]

	Cine Camera: Fix crash in CineCameraComponent when setting focus distance from BP

Change 3047366 on 2016/07/12 by [REDACTED_PERSON]

	Sequence Recorder: Arbitrary property recording

diff --git a/Engine/Plugins/MovieScene/LevelSequenceEditor/Source/LevelSequenceEditor/Private/LevelSequenceEditorToolkit.cpp b/Engine/Plugins/MovieScene/LevelSequenceEditor/Source/LevelSequenceEditor/Private/LevelSequenceEditorToolkit.cpp
--- a/Engine/Plugins/MovieScene/LevelSequenceEditor/Source/LevelSequenceEditor/Private/LevelSequenceEditorToolkit.cpp
+++ b/Engine/Plugins/MovieScene/LevelSequenceEditor/Source/LevelSequenceEditor/Private/LevelSequenceEditorToolkit.cpp
@@ -25,6 +25,7 @@
 #include "SceneOutlinerModule.h"
 #include "SceneOutlinerPublicTypes.h"
 #include "SDockTab.h"
+#include "LevelSequencePlayer.h"
 #include "SequencerSettings.h"
 #include "SequencerSpawnRegister.h"
 
@@ -46,6 +47,37 @@ namespace SequencerDefs
 	static const FName SequencerAppIdentifier(TEXT("SequencerApp"));
 }
 
+// Defer to ULevelSequencePlayer's implementation for getting event contexts from the current world
+TArray<UObject*> GetLevelSequenceEditorEventContexts()
+{
+	TArray<UObject*> Contexts;
+
+	// Return PIE worlds if there are any
+	for (const FWorldContext& Context : GEngine->GetWorldContexts())
+	{
+		if (Context.WorldType == EWorldType::PIE)
+		{
+			ULevelSequencePlayer::GetEventContexts(*Context.World(), Contexts);
+		}
+	}
+
+	if (Contexts.Num())
+	{
+		return Contexts;
+	}
+
+	// Else just return the editor world
+	for (const FWorldContext& Context : GEngine->GetWorldContexts())
+	{
+		if (Context.WorldType == EWorldType::Editor)
+		{
+			ULevelSequencePlayer::GetEventContexts(*Context.World(), Contexts);
+			break;
+		}
+	}
+
+	return Contexts;
+}
 
 static TArray<FLevelSequenceEditorToolkit*> OpenToolkits;
 
@@ -141,6 +173,7 @@ void FLevelSequenceEditorToolkit::Initialize(const EToolkitMode::Type Mode, cons
 		SequencerInitParams.bEditWithinLevelEditor = bEditWithinLevelEditor;
 		SequencerInitParams.ToolkitHost = InitToolkitHost;
 		SequencerInitParams.SpawnRegister = SpawnRegister;
+		SequencerInitParams.EventContexts.BindStatic(GetLevelSequenceEditorEventContexts);
 
 		TSharedRef<FExtender> AddMenuExtender = MakeShareable(new FExtender);
 
@@ -490,6 +523,7 @@ void FLevelSequenceEditorToolkit::AddDefaultTracksForActor(AActor& Actor, const
 			}
 
 			Sequencer->KeyProperty(KeyPropertyParams);
+
 			Sequencer->UpdateRuntimeInstances();
 		}
 	}
@@ -600,7 +634,7 @@ void FLevelSequenceEditorToolkit::HandleAddComponentMaterialActionExecute(UPrimi
 		MaterialTrack->Modify();
 		MaterialTrack->SetMaterialIndex( MaterialIndex );
 
-		Sequencer->NotifyMovieSceneDataChanged();
+		Sequencer->NotifyMovieSceneDataChanged( EMovieSceneDataChangeType::MovieSceneStructureItemAdded );
 	}
 }
 
@@ -645,7 +679,7 @@ void FLevelSequenceEditorToolkit::AddShot(UMovieSceneCinematicShotTrack* ShotTra
 			UObject* SpawnedCamera = GetSequencer()->FindSpawnedObjectOrTemplate(CameraGuid);
 			if (SpawnedCamera)
 			{
-				GWorld->EditorDestroyActor(NewCamera, true);
+				GCurrentLevelEditingViewportClient->GetWorld()->EditorDestroyActor(NewCamera, true);
 				NewCamera = Cast<ACineCameraActor>(SpawnedCamera);
 			}
 		}
diff --git a/Engine/Plugins/MovieScene/LevelSequenceEditor/Source/LevelSequenceEditor/Private/Misc/LevelSequenceEditorSpawnRegister.cpp b/Engine/Plugins/MovieScene/LevelSequenceEditor/Source/LevelSequenceEditor/Private/Misc/LevelSequenceEditorSpawnRegister.cpp
--- a/Engine/Plugins/MovieScene/LevelSequenceEditor/Source/LevelSequenceEditor/Private/Misc/LevelSequenceEditorSpawnRegister.cpp
+++ b/Engine/Plugins/MovieScene/LevelSequenceEditor/Source/LevelSequenceEditor/Private/Misc/LevelSequenceEditorSpawnRegister.cpp
@@ -115,12 +115,11 @@ void FLevelSequenceEditorSpawnRegister::SaveDefaultSpawnableState(const FGuid& B
 	}
 
 	// Find the spawnable definition
-	UMovieScene* OwnerMovieScene = SequenceInstance.GetSequence()->GetMovieScene();
-	FMovieSceneSpawnable* Spawnable = OwnerMovieScene->FindSpawnable(BindingId);
+	FMovieSceneSpawnable* Spawnable = SequenceInstance.GetSequence()->GetMovieScene()->FindSpawnable(BindingId);
 	if (Spawnable)
 	{
 		SequenceInstance.RestoreSpecificState(BindingId, *Sequencer);
-		Spawnable->CopyObjectTemplate(*Object, *OwnerMovieScene);
+		Spawnable->CopyObjectTemplate(*Object, *SequenceInstance.GetSequence());
 	}
 }
 
diff --git a/Engine/Plugins/MovieScene/MatineeToLevelSequence/Source/MatineeToLevelSequence/Private/MatineeToLevelSequenceModule.cpp b/Engine/Plugins/MovieScene/MatineeToLevelSequence/Source/MatineeToLevelSequence/Private/MatineeToLevelSequenceModule.cpp
--- a/Engine/Plugins/MovieScene/MatineeToLevelSequence/Source/MatineeToLevelSequence/Private/MatineeToLevelSequenceModule.cpp
+++ b/Engine/Plugins/MovieScene/MatineeToLevelSequence/Source/MatineeToLevelSequence/Private/MatineeToLevelSequenceModule.cpp
@@ -284,8 +284,9 @@ protected:
 			ObjectGuid = NewSequence->FindPossessableObjectId(*PropObject);
 			if (!ObjectGuid.IsValid())
 			{
+				UObject* BindingContext = InActor->GetWorld();
 				ObjectGuid = NewMovieScene->AddPossessable(PropObject->GetName(), PropObject->GetClass());
-				NewSequence->BindPossessableObject(ObjectGuid, *PropObject, GWorld);
+				NewSequence->BindPossessableObject(ObjectGuid, *PropObject, BindingContext);
 			}
 
 			// cbb: String manipulations to get the property path in the rigth form for sequencer
@@ -353,8 +354,9 @@ protected:
 		// Bind the group actor as a possessable						
 		if (GroupActor)
 		{
+			UObject* BindingContext = GroupActor->GetWorld();
 			PossessableGuid = NewMovieScene->AddPossessable(GroupActor->GetActorLabel(), GroupActor->GetClass());
-			NewSequence->BindPossessableObject(PossessableGuid, *GroupActor, GWorld);
+			NewSequence->BindPossessableObject(PossessableGuid, *GroupActor, BindingContext);
 	
 			FindOrAddFolder(NewMovieScene, GroupActor, PossessableGuid);
 		}
@@ -460,6 +462,18 @@ protected:
 					}
 				}
 			}
+			else if (Track->IsA(UInterpTrackLinearColorProp::StaticClass()))
+			{
+				UInterpTrackLinearColorProp* MatineeLinearColorTrack = StaticCast<UInterpTrackLinearColorProp*>(Track);
+				if (MatineeLinearColorTrack->GetNumKeyframes() != 0 && GroupActor && PossessableGuid.IsValid())
+				{
+					UMovieSceneColorTrack* ColorTrack = AddPropertyTrack<UMovieSceneColorTrack>(MatineeLinearColorTrack->PropertyName, GroupActor, PossessableGuid, NewSequence, NewMovieScene, NumWarnings);
+					if (ColorTrack)
+					{
+						FMatineeImportTools::CopyInterpLinearColorTrack(MatineeLinearColorTrack, ColorTrack);
+					}
+				}
+			}
 			else if (Track->IsA(UInterpTrackVisibility::StaticClass()))
 			{
 				UInterpTrackVisibility* MatineeVisibilityTrack = StaticCast<UInterpTrackVisibility*>(Track);
diff --git a/Engine/Plugins/MovieScene/MatineeToLevelSequence/Source/MatineeToLevelSequence/Private/MatineeToLevelSequencePCH.h b/Engine/Plugins/MovieScene/MatineeToLevelSequence/Source/MatineeToLevelSequence/Private/MatineeToLevelSequencePCH.h
--- a/Engine/Plugins/MovieScene/MatineeToLevelSequence/Source/MatineeToLevelSequence/Private/MatineeToLevelSequencePCH.h
+++ b/Engine/Plugins/MovieScene/MatineeToLevelSequence/Source/MatineeToLevelSequence/Private/MatineeToLevelSequencePCH.h
@@ -20,6 +20,7 @@
 #include "Matinee/InterpData.h"
 #include "Matinee/InterpTrackFloatBase.h"
 #include "Matinee/InterpTrackBoolProp.h"
+#include "Matinee/InterpTrackLinearColorProp.h"
 #include "Matinee/InterpTrackColorScale.h"
 #include "Matinee/InterpTrackColorProp.h"
 #include "Matinee/InterpTrackFloatProp.h"
diff --git a/Engine/Source/Editor/Kismet/Private/BlueprintDetailsCustomization.cpp b/Engine/Source/Editor/Kismet/Private/BlueprintDetailsCustomization.cpp
--- a/Engine/Source/Editor/Kismet/Private/BlueprintDetailsCustomization.cpp
+++ b/Engine/Source/Editor/Kismet/Private/BlueprintDetailsCustomization.cpp
@@ -363,24 +363,24 @@ void FBlueprintVarActionDetails::CustomizeDetails( IDetailLayoutBuilder& DetailL
 		.ToolTip(PrivateTooltip)
 	];
 
-	TSharedPtr<SToolTip> ExposeToMatineeTooltip = IDocumentation::Get()->CreateToolTip(LOCTEXT("VariableExposeToMatinee_Tooltip", "Should this variable be exposed for Matinee to modify?"), NULL, DocLink, TEXT("ExposeToMatinee"));
+	TSharedPtr<SToolTip> ExposeToCinematicsTooltip = IDocumentation::Get()->CreateToolTip(LOCTEXT("VariableExposeToCinematics_Tooltip", "Should this variable be exposed for Matinee or Sequencer to modify?"), NULL, DocLink, TEXT("ExposeToCinematics"));
 
-	Category.AddCustomRow( LOCTEXT("VariableExposeToMatinee", "Expose to Matinee") )
-	.Visibility(TAttribute<EVisibility>(this, &FBlueprintVarActionDetails::ExposeToMatineeVisibility))
+	Category.AddCustomRow( LOCTEXT("VariableExposeToCinematics", "Expose to Cinematics") )
+	.Visibility(TAttribute<EVisibility>(this, &FBlueprintVarActionDetails::ExposeToCinematicsVisibility))
 	.NameContent()
 	[
 		SNew(STextBlock)
-		.ToolTip(ExposeToMatineeTooltip)
-		.Text( LOCTEXT("VariableExposeToMatinee", "Expose to Matinee") )
+		.ToolTip(ExposeToCinematicsTooltip)
+		.Text( LOCTEXT("VariableExposeToCinematics", "Expose to Cinematics") )
 		.Font( DetailFontInfo )
 	]
 	.ValueContent()
 	[
 		SNew(SCheckBox)
-		.IsChecked( this, &FBlueprintVarActionDetails::OnGetExposedToMatineeCheckboxState )
-		.OnCheckStateChanged( this, &FBlueprintVarActionDetails::OnExposedToMatineeChanged )
+		.IsChecked( this, &FBlueprintVarActionDetails::OnGetExposedToCinematicsCheckboxState )
+		.OnCheckStateChanged( this, &FBlueprintVarActionDetails::OnExposedToCinematicsChanged )
 		.IsEnabled(IsVariableInBlueprint())
-		.ToolTip(ExposeToMatineeTooltip)
+		.ToolTip(ExposeToCinematicsTooltip)
 	];
 
 	// Build the property specific config variable tool tip
@@ -1619,7 +1619,7 @@ EVisibility FBlueprintVarActionDetails::ExposePrivateVisibility() const
 	return EVisibility::Collapsed;
 }
 
-ECheckBoxState FBlueprintVarActionDetails::OnGetExposedToMatineeCheckboxState() const
+ECheckBoxState FBlueprintVarActionDetails::OnGetExposedToCinematicsCheckboxState() const
 {
 	UProperty* Property = CachedVariableProperty.Get();
 	if (Property)
@@ -1629,19 +1629,19 @@ ECheckBoxState FBlueprintVarActionDetails::OnGetExposedToMatineeCheckboxState()
 	return ECheckBoxState::Unchecked;
 }
 
-void FBlueprintVarActionDetails::OnExposedToMatineeChanged(ECheckBoxState InNewState)
+void FBlueprintVarActionDetails::OnExposedToCinematicsChanged(ECheckBoxState InNewState)
 {
 	// Toggle the flag on the blueprint's version of the variable description, based on state
-	const bool bExposeToMatinee = (InNewState == ECheckBoxState::Checked);
+	const bool bExposeToCinematics = (InNewState == ECheckBoxState::Checked);
 	
 	const FName VarName = CachedVariableName;
 	if (VarName != NAME_None)
 	{
-		FBlueprintEditorUtils::SetInterpFlag(GetBlueprintObj(), VarName, bExposeToMatinee);
+		FBlueprintEditorUtils::SetInterpFlag(GetBlueprintObj(), VarName, bExposeToCinematics);
 	}
 }
 
-EVisibility FBlueprintVarActionDetails::ExposeToMatineeVisibility() const
+EVisibility FBlueprintVarActionDetails::ExposeToCinematicsVisibility() const
 {
 	UProperty* VariableProperty = CachedVariableProperty.Get();
 	if (VariableProperty && !IsALocalVariable(VariableProperty))
diff --git a/Engine/Source/Editor/Kismet/Private/BlueprintDetailsCustomization.h b/Engine/Source/Editor/Kismet/Private/BlueprintDetailsCustomization.h
--- a/Engine/Source/Editor/Kismet/Private/BlueprintDetailsCustomization.h
+++ b/Engine/Source/Editor/Kismet/Private/BlueprintDetailsCustomization.h
@@ -144,9 +144,9 @@ private:
 	void OnPrivateChanged(ECheckBoxState InNewState);
 	EVisibility ExposePrivateVisibility() const;
 	
-	ECheckBoxState OnGetExposedToMatineeCheckboxState() const;
-	void OnExposedToMatineeChanged(ECheckBoxState InNewState);
-	EVisibility ExposeToMatineeVisibility() const;
+	ECheckBoxState OnGetExposedToCinematicsCheckboxState() const;
+	void OnExposedToCinematicsChanged(ECheckBoxState InNewState);
+	EVisibility ExposeToCinematicsVisibility() const;
 
 	ECheckBoxState OnGetConfigVariableCheckboxState() const;
 	void OnSetConfigVariableState(ECheckBoxState InNewState);
diff --git a/Engine/Source/Editor/Matinee/Private/MatineeTools.cpp b/Engine/Source/Editor/Matinee/Private/MatineeTools.cpp
--- a/Engine/Source/Editor/Matinee/Private/MatineeTools.cpp
+++ b/Engine/Source/Editor/Matinee/Private/MatineeTools.cpp
@@ -3896,11 +3896,11 @@ void FMatinee::UpdateLevelViewport(AActor* InActor, FLevelEditorViewportClient*
 		if (PreviousCamera != Cam)
 		{
 			PreviousCamera = Cam;
-			InViewportClient->bEditorCameraCut = true;
+			InViewportClient->SetIsCameraCut(true);
 		}
 		else
 		{
-			InViewportClient->bEditorCameraCut = false;
+			InViewportClient->SetIsCameraCut(false);
 		}
 	}
 	else
@@ -3910,7 +3910,7 @@ void FMatinee::UpdateLevelViewport(AActor* InActor, FLevelEditorViewportClient*
 		InViewportClient->FadeAmount = InFadeAmount;
 		InViewportClient->bEnableFading = true;
 
-		InViewportClient->bEditorCameraCut = false;
+		InViewportClient->SetIsCameraCut(false);
 	}
 
 	// Set the actor lock.
diff --git a/Engine/Source/Editor/MovieSceneCaptureDialog/Private/MovieSceneCaptureDialogModule.cpp b/Engine/Source/Editor/MovieSceneCaptureDialog/Private/MovieSceneCaptureDialogModule.cpp
--- a/Engine/Source/Editor/MovieSceneCaptureDialog/Private/MovieSceneCaptureDialogModule.cpp
+++ b/Engine/Source/Editor/MovieSceneCaptureDialog/Private/MovieSceneCaptureDialogModule.cpp
@@ -26,6 +26,8 @@
 #include "ISessionInfo.h"
 #include "ISessionManager.h"
 
+#include "ErrorCodes.h"
+
 #include "SLevelViewport.h"
 
 #define LOCTEXT_NAMESPACE "MovieSceneCaptureDialog"
@@ -136,12 +138,35 @@ enum class ECaptureState
 
 DECLARE_DELEGATE_OneParam(FOnCaptureFinished, bool /*bCancelled*/);
 
+// Structure used to store the state of the capture
+struct FCaptureState
+{
+	/** Construction from an enum */
+	explicit FCaptureState(ECaptureState InState = ECaptureState::Pending) : State(InState), Code(0){}
+	/** Construction from a process exit code */
+	explicit FCaptureState(int32 InCode) : State(InCode == 0 ? ECaptureState::Success : ECaptureState::Failure), Code(InCode){}
+
+	/** Get any additional detailed text */
+	FText GetDetailText()
+	{
+		switch(uint32(Code))
+		{
+		case uint32(EMovieSceneCaptureExitCode::WorldNotFound): return LOCTEXT("WorldNotFound", "Specified world does not exist. Did you forget to save it?");
+		}
+
+		return FText();
+	}
+
+	ECaptureState State;
+	int32 Code;
+};
+
 class SCaptureMovieNotification : public SCompoundWidget, public INotificationWidget
 {
 public:
 	SLATE_BEGIN_ARGS(SCaptureMovieNotification){}
 
-		SLATE_ATTRIBUTE(ECaptureState, CaptureState)
+		SLATE_ATTRIBUTE(FCaptureState, CaptureState)
 
 		SLATE_EVENT(FOnCaptureFinished, OnCaptureFinished)
 
@@ -157,7 +182,7 @@ public:
 		OnCaptureFinished = InArgs._OnCaptureFinished;
 		OnCancel = InArgs._OnCancel;
 
-		CachedState = ECaptureState::Pending;
+		CachedState = FCaptureState(ECaptureState::Pending);
 
 		FString CapturePath = FPaths::ConvertRelativePathToFull(InArgs._CapturePath);
 		CapturePath.RemoveFromEnd(TEXT("\\"));
@@ -197,6 +222,16 @@ public:
 					]
 				]
 
+				+ SVerticalBox::Slot()
+				.AutoHeight()
+				.Padding(FMargin(0,0,0,5.0f))
+				.HAlign(HAlign_Right)
+				[
+					SAssignNew(DetailedTextBlock, STextBlock)
+					.Visibility(EVisibility::Collapsed)
+					.Font(FCoreStyle::Get().GetFontStyle(TEXT("NotificationList.FontLight")))
+				]
+
 				+ SVerticalBox::Slot()
 				.AutoHeight()
 				.HAlign(HAlign_Right)
@@ -233,20 +268,26 @@ public:
 			return;
 		}
 
-		ECaptureState StateThisFrame = CaptureState.Get();
+		FCaptureState StateThisFrame = CaptureState.Get();
 
-		if (CachedState != StateThisFrame)
+		if (CachedState.State != StateThisFrame.State)
 		{
 			CachedState = StateThisFrame;
 			
-			if (CachedState == ECaptureState::Success)
+			if (CachedState.State == ECaptureState::Success)
 			{
 				TextBlock->SetText(LOCTEXT("CaptureFinished", "Capture Finished"));
 				OnCaptureFinished.ExecuteIfBound(true);
 			}
-			else if (CachedState == ECaptureState::Failure)
+			else if (CachedState.State == ECaptureState::Failure)
 			{
 				TextBlock->SetText(LOCTEXT("CaptureFailed", "Capture Failed"));
+				FText DetailText = CachedState.GetDetailText();
+				if (!DetailText.IsEmpty())
+				{
+					DetailedTextBlock->SetText(DetailText);
+					DetailedTextBlock->SetVisibility(EVisibility::Visible);
+				}
 				OnCaptureFinished.ExecuteIfBound(false);
 			}
 			else
@@ -286,11 +327,12 @@ private:
 private:
 	TSharedPtr<SWidget> Button, Throbber, Hyperlink;
 	TSharedPtr<STextBlock> TextBlock;
+	TSharedPtr<STextBlock> DetailedTextBlock;
 	SNotificationItem::ECompletionState State;
 
 	FSimpleDelegate OnCancel;
-	ECaptureState CachedState;
-	TAttribute<ECaptureState> CaptureState;
+	FCaptureState CachedState;
+	TAttribute<FCaptureState> CaptureState;
 	FOnCaptureFinished OnCaptureFinished;
 };
 
@@ -332,7 +374,7 @@ struct FInEditorCapture
 		UGameViewportClient::OnViewportCreated().AddRaw(this, &FInEditorCapture::OnStart);
 		FEditorDelegates::EndPIE.AddRaw(this, &FInEditorCapture::OnEndPIE);
 		
-		FAudioDevice* AudioDevice = GWorld->GetAudioDevice();
+		FAudioDevice* AudioDevice = GEngine->GetMainAudioDevice();
 		if (AudioDevice != nullptr)
 		{
 			TransientMasterVolume = AudioDevice->GetTransientMasterVolume();
@@ -451,7 +493,7 @@ struct FInEditorCapture
 
 		FObjectReader(GetMutableDefault<ULevelEditorPlaySettings>(), BackedUpPlaySettings);
 
-		FAudioDevice* AudioDevice = GWorld->GetAudioDevice();
+		FAudioDevice* AudioDevice = GEngine->GetMainAudioDevice();
 		if (AudioDevice != nullptr)
 		{
 			AudioDevice->SetTransientMasterVolume(TransientMasterVolume);
@@ -596,11 +638,11 @@ class FMovieSceneCaptureDialogModule : public IMovieSceneCaptureDialogModule
 			{
 				if (Context.WorldType == EWorldType::PIE)
 				{
-					return ECaptureState::Pending;
+					return FCaptureState(ECaptureState::Pending);
 				}
 			}
 
-			return ECaptureState::Success;
+			return FCaptureState(ECaptureState::Success);
 		};
 
 		auto OnCaptureStarted = [=]{
@@ -654,7 +696,7 @@ class FMovieSceneCaptureDialogModule : public IMovieSceneCaptureDialogModule
 			return LOCTEXT("UnableToSaveCaptureManifest", "Unable to save capture manifest");
 		}
 
-		FString EditorCommandLine = FString::Printf(TEXT("%s -MovieSceneCaptureManifest=\"%s\" -game"), *MapNameToLoad, *Filename);
+		FString EditorCommandLine = FString::Printf(TEXT("%s -MovieSceneCaptureManifest=\"%s\" -game -NoLoadingScreen -ForceRes -Windowed"), *MapNameToLoad, *Filename);
 
 		if( CaptureObject->Settings.bCreateTemporaryCopiesOfLevels )
 		{
@@ -718,11 +760,11 @@ class FMovieSceneCaptureDialogModule : public IMovieSceneCaptureDialogModule
 				{
 					int32 RetCode = 0;
 					FPlatformProcess::GetProcReturnCode(*SharedProcHandle, &RetCode);
-					return RetCode == 0 ? ECaptureState::Success : ECaptureState::Failure;
+					return FCaptureState(RetCode);
 				}
 				else
 				{
-					return ECaptureState::Pending;
+					return FCaptureState(ECaptureState::Pending);
 				}
 			};
 
diff --git a/Engine/Source/Editor/MovieSceneTools/Private/CurveKeyEditors/SBoolCurveKeyEditor.cpp b/Engine/Source/Editor/MovieSceneTools/Private/CurveKeyEditors/SBoolCurveKeyEditor.cpp
--- a/Engine/Source/Editor/MovieSceneTools/Private/CurveKeyEditors/SBoolCurveKeyEditor.cpp
+++ b/Engine/Source/Editor/MovieSceneTools/Private/CurveKeyEditors/SBoolCurveKeyEditor.cpp
@@ -70,7 +70,7 @@ void SBoolCurveKeyEditor::OnCheckStateChanged(ECheckBoxState NewCheckboxState)
 
 		OnValueChangedEvent.ExecuteIfBound(NewCheckboxState == ECheckBoxState::Checked);
 
-		Sequencer->UpdateRuntimeInstances();
+		Sequencer->NotifyMovieSceneDataChanged( EMovieSceneDataChangeType::TrackValueChanged );
 	}
 }
 
diff --git a/Engine/Source/Editor/MovieSceneTools/Private/CurveKeyEditors/SEnumCurveKeyEditor.cpp b/Engine/Source/Editor/MovieSceneTools/Private/CurveKeyEditors/SEnumCurveKeyEditor.cpp
--- a/Engine/Source/Editor/MovieSceneTools/Private/CurveKeyEditors/SEnumCurveKeyEditor.cpp
+++ b/Engine/Source/Editor/MovieSceneTools/Private/CurveKeyEditors/SEnumCurveKeyEditor.cpp
@@ -60,7 +60,7 @@ void SEnumCurveKeyEditor::OnComboSelectionChanged(int32 InSelectedItem, ESelectI
 
 		OnValueChangedEvent.ExecuteIfBound(InSelectedItem);
 
-		Sequencer->UpdateRuntimeInstances();
+		Sequencer->NotifyMovieSceneDataChanged( EMovieSceneDataChangeType::TrackValueChanged );
 	}
 }
 
diff --git a/Engine/Source/Editor/MovieSceneTools/Private/CurveKeyEditors/SFloatCurveKeyEditor.cpp b/Engine/Source/Editor/MovieSceneTools/Private/CurveKeyEditors/SFloatCurveKeyEditor.cpp
--- a/Engine/Source/Editor/MovieSceneTools/Private/CurveKeyEditors/SFloatCurveKeyEditor.cpp
+++ b/Engine/Source/Editor/MovieSceneTools/Private/CurveKeyEditors/SFloatCurveKeyEditor.cpp
@@ -107,13 +107,13 @@ void SFloatCurveKeyEditor::OnValueChanged(float Value)
 
 		OnValueChangedEvent.ExecuteIfBound(Value);
 
-		Sequencer->UpdateRuntimeInstances();
+		Sequencer->NotifyMovieSceneDataChanged( EMovieSceneDataChangeType::TrackValueChanged );
 	}
 }
 
 void SFloatCurveKeyEditor::OnValueCommitted(float Value, ETextCommit::Type CommitInfo)
 {
-	if (CommitInfo == ETextCommit::OnEnter)
+	if (CommitInfo == ETextCommit::OnEnter || CommitInfo == ETextCommit::OnUserMovedFocus)
 	{
 		const FScopedTransaction Transaction( LOCTEXT("SetFloatKey", "Set Float Key Value") );
 
diff --git a/Engine/Source/Editor/MovieSceneTools/Private/CurveKeyEditors/SIntegralCurveKeyEditor.cpp b/Engine/Source/Editor/MovieSceneTools/Private/CurveKeyEditors/SIntegralCurveKeyEditor.cpp
--- a/Engine/Source/Editor/MovieSceneTools/Private/CurveKeyEditors/SIntegralCurveKeyEditor.cpp
+++ b/Engine/Source/Editor/MovieSceneTools/Private/CurveKeyEditors/SIntegralCurveKeyEditor.cpp
@@ -83,13 +83,13 @@ void SIntegralCurveKeyEditor::OnValueChanged(int32 Value)
 		{
 			Curve->UpdateOrAddKey(CurrentTime, Value);
 		}
-		Sequencer->UpdateRuntimeInstances();
+		Sequencer->NotifyMovieSceneDataChanged( EMovieSceneDataChangeType::TrackValueChanged );
 	}
 }
 
 void SIntegralCurveKeyEditor::OnValueCommitted(int32 Value, ETextCommit::Type CommitInfo)
 {
-	if (CommitInfo == ETextCommit::OnEnter)
+	if (CommitInfo == ETextCommit::OnEnter || CommitInfo == ETextCommit::OnUserMovedFocus)
 	{
 		const FScopedTransaction Transaction( LOCTEXT("SetIntegralKey", "Set Integral Key Value") );
 
diff --git a/Engine/Source/Editor/MovieSceneTools/Private/MatineeImportTools.cpp b/Engine/Source/Editor/MovieSceneTools/Private/MatineeImportTools.cpp
--- a/Engine/Source/Editor/MovieSceneTools/Private/MatineeImportTools.cpp
+++ b/Engine/Source/Editor/MovieSceneTools/Private/MatineeImportTools.cpp
@@ -6,6 +6,7 @@
 #include "Matinee/MatineeActor.h"
 #include "Matinee/InterpData.h"
 #include "Matinee/InterpGroupInst.h"
+#include "Matinee/InterpTrackLinearColorProp.h"
 #include "Matinee/InterpTrackColorProp.h"
 #include "Matinee/InterpTrackBoolProp.h"
 #include "Matinee/InterpTrackFloatBase.h"
@@ -228,6 +229,55 @@ bool FMatineeImportTools::CopyInterpColorTrack( UInterpTrackColorProp* ColorProp
 	return bSectionCreated;
 }
 
+
+bool FMatineeImportTools::CopyInterpLinearColorTrack( UInterpTrackLinearColorProp* LinearColorPropTrack, UMovieSceneColorTrack* ColorTrack )
+{
+	const FScopedTransaction Transaction( NSLOCTEXT( "Sequencer", "PasteMatineeColorTrack", "Paste Matinee Linear Color Track" ) );
+	bool bSectionCreated = false;
+
+	ColorTrack->Modify();
+
+	float KeyTime = LinearColorPropTrack->GetKeyframeTime( 0 );
+	UMovieSceneColorSection* Section = Cast<UMovieSceneColorSection>( MovieSceneHelpers::FindSectionAtTime( ColorTrack->GetAllSections(), KeyTime ) );
+	if ( Section == nullptr )
+	{
+		Section = Cast<UMovieSceneColorSection>( ColorTrack->CreateNewSection() );
+		ColorTrack->AddSection( *Section );
+		Section->GetRedCurve().SetDefaultValue(0.f);
+		Section->GetGreenCurve().SetDefaultValue(0.f);
+		Section->GetBlueCurve().SetDefaultValue(0.f);
+		Section->GetAlphaCurve().SetDefaultValue(1.f);
+		Section->SetIsInfinite(true);
+		bSectionCreated = true;
+	}
+
+	if (Section->TryModify())
+	{
+		float SectionMin = Section->GetStartTime();
+		float SectionMax = Section->GetEndTime();
+
+		FRichCurve& RedCurve = Section->GetRedCurve();
+		FRichCurve& GreenCurve = Section->GetGreenCurve();
+		FRichCurve& BlueCurve = Section->GetBlueCurve();
+		FRichCurve& AlphaCurve = Section->GetAlphaCurve();
+
+		for ( const auto& Point : LinearColorPropTrack->LinearColorTrack.Points)
+		{
+			FMatineeImportTools::SetOrAddKey( RedCurve, Point.InVal, Point.OutVal.R, Point.ArriveTangent.R, Point.LeaveTangent.R, Point.InterpMode );
+			FMatineeImportTools::SetOrAddKey( GreenCurve, Point.InVal, Point.OutVal.G, Point.ArriveTangent.G, Point.LeaveTangent.G, Point.InterpMode );
+			FMatineeImportTools::SetOrAddKey( BlueCurve, Point.InVal, Point.OutVal.B, Point.ArriveTangent.B, Point.LeaveTangent.B, Point.InterpMode );
+			FMatineeImportTools::SetOrAddKey( AlphaCurve, Point.InVal, Point.OutVal.A, Point.ArriveTangent.A, Point.LeaveTangent.A, Point.InterpMode );
+			SectionMin = FMath::Min( SectionMin, Point.InVal );
+			SectionMax = FMath::Max( SectionMax, Point.InVal );
+		}
+
+		Section->SetStartTime( SectionMin );
+		Section->SetEndTime( SectionMax );
+	}
+
+	return bSectionCreated;
+}
+
 bool FMatineeImportTools::CopyInterpMoveTrack( UInterpTrackMove* MoveTrack, UMovieScene3DTransformTrack* TransformTrack )
 {
 	const FScopedTransaction Transaction( NSLOCTEXT( "Sequencer", "PasteMatineeMoveTrack", "Paste Matinee Move Track" ) );
diff --git a/Engine/Source/Editor/MovieSceneTools/Private/MovieSceneToolHelpers.cpp b/Engine/Source/Editor/MovieSceneTools/Private/MovieSceneToolHelpers.cpp
--- a/Engine/Source/Editor/MovieSceneTools/Private/MovieSceneToolHelpers.cpp
+++ b/Engine/Source/Editor/MovieSceneTools/Private/MovieSceneToolHelpers.cpp
@@ -555,45 +555,54 @@ bool MovieSceneToolHelpers::ShowImportEDLDialog(UMovieScene* InMovieScene, float
 
 bool MovieSceneToolHelpers::ShowExportEDLDialog(const UMovieScene* InMovieScene, float InFrameRate, FString InSaveDirectory)
 {
+	TArray<FString> SaveFilenames;
 	FString SequenceName = InMovieScene->GetOuter()->GetName();
-	FString SequencePath = InMovieScene->GetOuter()->GetPathName();
-	SequencePath = FPaths::GetPath(SequencePath);
 
 	// Pop open a dialog to request the location of the edl
-	FString SaveDirectory = InSaveDirectory;
-	if (InSaveDirectory.IsEmpty())
+	IDesktopPlatform* DesktopPlatform = FDesktopPlatformModule::Get();
+	bool bSave = false;
+	if (DesktopPlatform)
 	{
-		IDesktopPlatform* DesktopPlatform = FDesktopPlatformModule::Get();
-		bool bSave = false;
-		if (DesktopPlatform)
-		{
-			void* ParentWindowWindowHandle = NULL;
-
-			IMainFrameModule& MainFrameModule = FModuleManager::LoadModuleChecked<IMainFrameModule>(TEXT("MainFrame"));
-			const TSharedPtr<SWindow>& MainFrameParentWindow = MainFrameModule.GetParentWindow();
-			if ( MainFrameParentWindow.IsValid() && MainFrameParentWindow->GetNativeWindow().IsValid() )
-			{
-				ParentWindowWindowHandle = MainFrameParentWindow->GetNativeWindow()->GetOSWindowHandle();
-			}
+		void* ParentWindowWindowHandle = NULL;
 
-			bSave = DesktopPlatform->OpenDirectoryDialog(
-				ParentWindowWindowHandle,
-				NSLOCTEXT("UnrealEd", "MovieSceneToolHelpersExportEDL", "Export EDL to...").ToString(), 
-				SequencePath,
-				SaveDirectory 
-				);
-		}
-		if (!bSave)
+		IMainFrameModule& MainFrameModule = FModuleManager::LoadModuleChecked<IMainFrameModule>(TEXT("MainFrame"));
+		const TSharedPtr<SWindow>& MainFrameParentWindow = MainFrameModule.GetParentWindow();
+		if ( MainFrameParentWindow.IsValid() && MainFrameParentWindow->GetNativeWindow().IsValid() )
 		{
-			return false;
+			ParentWindowWindowHandle = MainFrameParentWindow->GetNativeWindow()->GetOSWindowHandle();
 		}
+
+		FString ExtensionStr;
+		ExtensionStr += TEXT("CMX 3600 EDL (*.edl)|*.edl|");
+		ExtensionStr += TEXT("RV (*.rv)|*.rv|");
+
+		bSave = DesktopPlatform->SaveFileDialog(
+			ParentWindowWindowHandle,
+			NSLOCTEXT("UnrealEd", "MovieSceneToolHelpersExportEDL", "Export EDL to...").ToString(), 
+			InSaveDirectory,
+			SequenceName + TEXT(".edl"), 
+			*ExtensionStr,
+			EFileDialogFlags::None,
+			SaveFilenames
+			);
+	}
+	if (!bSave)
+	{
+		return false;
 	}
 
-	if (MovieSceneCaptureHelpers::ExportEDL(InMovieScene, InFrameRate, SaveDirectory))
+	if (!SaveFilenames.Num())
 	{
+		return false;
+	}
+
+	if (MovieSceneCaptureHelpers::ExportEDL(InMovieScene, InFrameRate, SaveFilenames[0]))
+	{
+		const FString AbsoluteFilename = FPaths::ConvertRelativePathToFull(SaveFilenames[0]);
+		const FString SaveDirectory = FPaths::GetPath(AbsoluteFilename);
+
 		FNotificationInfo NotificationInfo(NSLOCTEXT("MovieSceneTools", "EDLExportFinished", "EDL Export finished"));
 		NotificationInfo.ExpireDuration = 5.f;
-		NotificationInfo.bFireAndForget = true;
 		NotificationInfo.Hyperlink = FSimpleDelegate::CreateStatic( [](FString InDirectory) { FPlatformProcess::ExploreFolder( *InDirectory ); }, SaveDirectory);
 		NotificationInfo.HyperlinkText = NSLOCTEXT("MovieSceneTools", "OpenEDLExportFolder", "Open EDL Export Folder...");
 		FSlateNotificationManager::Get().AddNotification(NotificationInfo);
diff --git a/Engine/Source/Editor/MovieSceneTools/Private/MovieSceneToolsUserSettings.cpp b/Engine/Source/Editor/MovieSceneTools/Private/MovieSceneToolsUserSettings.cpp
--- a/Engine/Source/Editor/MovieSceneTools/Private/MovieSceneToolsUserSettings.cpp
+++ b/Engine/Source/Editor/MovieSceneTools/Private/MovieSceneToolsUserSettings.cpp
@@ -13,12 +13,19 @@ UMovieSceneUserThumbnailSettings::UMovieSceneUserThumbnailSettings(const FObject
 {
 	ThumbnailSize = FIntPoint(128, 72);
 	bDrawThumbnails = true;
+	Quality = EThumbnailQuality::Normal;
 }
 
 void UMovieSceneUserThumbnailSettings::PostEditChangeProperty(FPropertyChangedEvent& PropertyChangedEvent)
 {
 	Super::PostEditChangeProperty(PropertyChangedEvent);
 
+	const FName PropertyName = PropertyChangedEvent.Property ? PropertyChangedEvent.Property->GetFName() : NAME_None;
+	if (PropertyName == GET_MEMBER_NAME_CHECKED(UMovieSceneUserThumbnailSettings, Quality))
+	{
+		BroadcastRedrawThumbnails();
+	}
+	
 	ThumbnailSize.X = FMath::Clamp(ThumbnailSize.X, 1, 1024);
 	ThumbnailSize.Y = FMath::Clamp(ThumbnailSize.Y, 1, 1024);
 
diff --git a/Engine/Source/Editor/MovieSceneTools/Private/MovieSceneToolsUserSettings.h b/Engine/Source/Editor/MovieSceneTools/Private/MovieSceneToolsUserSettings.h
--- a/Engine/Source/Editor/MovieSceneTools/Private/MovieSceneToolsUserSettings.h
+++ b/Engine/Source/Editor/MovieSceneTools/Private/MovieSceneToolsUserSettings.h
@@ -5,6 +5,14 @@
 #include "IDetailCustomization.h"
 #include "MovieSceneToolsUserSettings.generated.h"
 
+UENUM()
+enum class EThumbnailQuality : uint8
+{
+	Draft,
+	Normal,
+	Best,
+};
+
 UCLASS(config=EditorSettings)
 class UMovieSceneUserThumbnailSettings : public UObject
 {
@@ -21,9 +29,14 @@ public:
 	UPROPERTY(EditAnywhere, config, Category=General, meta=(EditCondition=bDrawThumbnails))
 	bool bDrawSingleThumbnails;
 
+	/** Size at which to draw thumbnails on thumbnail sections */
 	UPROPERTY(EditAnywhere, config, Category=General, meta=(ClampMin=1, ClampMax=1024, EditCondition=bDrawThumbnails))
 	FIntPoint ThumbnailSize;
 
+	/** Quality to render the thumbnails with */
+	UPROPERTY(EditAnywhere, config, Category=General, meta=(EditCondition=bDrawThumbnails))
+	EThumbnailQuality Quality;
+
 	DECLARE_EVENT(UMovieSceneUserThumbnailSettings, FOnForceRedraw)
 	FOnForceRedraw& OnForceRedraw() { return OnForceRedrawEvent; }
 	void BroadcastRedrawThumbnails() const { OnForceRedrawEvent.Broadcast(); }
diff --git a/Engine/Source/Editor/MovieSceneTools/Private/Sections/CameraCutSection.cpp b/Engine/Source/Editor/MovieSceneTools/Private/Sections/CameraCutSection.cpp
--- a/Engine/Source/Editor/MovieSceneTools/Private/Sections/CameraCutSection.cpp
+++ b/Engine/Source/Editor/MovieSceneTools/Private/Sections/CameraCutSection.cpp
@@ -192,7 +192,7 @@ void FCameraCutSection::HandleSetCameraMenuEntryExecute(AActor* InCamera)
 	
 		CameraCutSection->SetCameraGuid(ObjectGuid);
 	
-		Sequencer->NotifyMovieSceneDataChanged();
+		Sequencer->NotifyMovieSceneDataChanged( EMovieSceneDataChangeType::TrackValueChanged );
 	}
 }
 
diff --git a/Engine/Source/Editor/MovieSceneTools/Private/Sections/ThumbnailSection.cpp b/Engine/Source/Editor/MovieSceneTools/Private/Sections/ThumbnailSection.cpp
--- a/Engine/Source/Editor/MovieSceneTools/Private/Sections/ThumbnailSection.cpp
+++ b/Engine/Source/Editor/MovieSceneTools/Private/Sections/ThumbnailSection.cpp
@@ -245,20 +245,8 @@ int32 FThumbnailSection::OnPaintSection( FSequencerSectionPainter& InPainter ) c
 				PaintGeometry,
 				Thumbnail,
 				ThumbnailClipRect,
-				DrawEffects | ESlateDrawEffect::NoGamma | ESlateDrawEffect::NoBlending
-				);
-		}
-
-		if (Fade > 0.0f)
-		{
-			FSlateDrawElement::MakeBox(
-				InPainter.DrawElements,
-				LayerId + 1,
-				PaintGeometry,
-				WhiteBrush,
-				ThumbnailClipRect,
-				DrawEffects,
-				FLinearColor(0.0f, 0.0f, 0.0f, Fade)
+				DrawEffects | ESlateDrawEffect::NoGamma,
+				FLinearColor(1.f, 1.f, 1.f, 1.f - Fade)
 				);
 		}
 	}
@@ -279,7 +267,7 @@ void FThumbnailSection::Tick(const FGeometry& AllottedGeometry, const FGeometry&
 		const TRange<float> VisibleRange = SequencerPtr.Pin()->GetViewRange();
 		const TRange<float> SectionRange = Section->IsInfinite() ? VisibleRange : Section->GetRange();
 
-		ThumbnailCache.Update(SectionRange, VisibleRange, AllocatedSize, Settings->ThumbnailSize);
+		ThumbnailCache.Update(SectionRange, VisibleRange, AllocatedSize, Settings->ThumbnailSize, Settings->Quality, InCurrentTime);
 	}
 }
 
diff --git a/Engine/Source/Editor/MovieSceneTools/Private/TrackEditorThumbnail/TrackEditorThumbnail.cpp b/Engine/Source/Editor/MovieSceneTools/Private/TrackEditorThumbnail/TrackEditorThumbnail.cpp
--- a/Engine/Source/Editor/MovieSceneTools/Private/TrackEditorThumbnail/TrackEditorThumbnail.cpp
+++ b/Engine/Source/Editor/MovieSceneTools/Private/TrackEditorThumbnail/TrackEditorThumbnail.cpp
@@ -136,13 +136,14 @@ void FTrackEditorThumbnail::CopyTextureIn(TSharedPtr<FSceneViewport> SceneViewpo
 		FTexture2DRHIRef SourceBackBuffer = RenderTarget->GetRHIRef();
 		PixelShader->SetParameters(RHICmdList, TStaticSamplerState<SF_Bilinear>::GetRHI(), SourceBackBuffer);
 
-		const float Left = (SourceBackBuffer->GetSizeX() - TargetTexture->GetWidth()) * .5f;
-		const float Top = (SourceBackBuffer->GetSizeY() - TargetTexture->GetHeight()) * .5f;
+		const float Scale = FMath::Min(float(SourceBackBuffer->GetSizeX()) / TargetTexture->GetWidth(), float(SourceBackBuffer->GetSizeY()) / TargetTexture->GetHeight());
+		const float Left = (SourceBackBuffer->GetSizeX() - TargetTexture->GetWidth()*Scale) * .5f;
+		const float Top = (SourceBackBuffer->GetSizeY() - TargetTexture->GetHeight()*Scale) * .5f;
 
 		const float U = Left / float(SourceBackBuffer->GetSizeX());
 		const float V = Top / float(SourceBackBuffer->GetSizeY());
-		const float SizeU = float(TargetTexture->GetWidth()) / float(SourceBackBuffer->GetSizeX());
-		const float SizeV = float(TargetTexture->GetHeight()) / float(SourceBackBuffer->GetSizeY());
+		const float SizeU = float(TargetTexture->GetWidth() * Scale) / float(SourceBackBuffer->GetSizeX());
+		const float SizeV = float(TargetTexture->GetHeight() * Scale) / float(SourceBackBuffer->GetSizeY());
 
 		RendererModule->DrawRectangle(
 			RHICmdList,
@@ -232,7 +233,9 @@ FTrackEditorThumbnailCache::FTrackEditorThumbnailCache(const TSharedPtr<FTrackEd
 	, ThumbnailPool(InThumbnailPool)
 {
 	FrameCount = 0;
+	LastComputationTime = 0;
 	bForceRedraw = false;
+	bNeedsNewThumbnails = false;
 }
 
 
@@ -255,21 +258,23 @@ void FTrackEditorThumbnailCache::SetSingleReferenceFrame(TOptional<float> InRefe
 	CurrentCache.SingleReferenceFrame = InReferenceFrame;
 }
 
-void FTrackEditorThumbnailCache::Update(const TRange<float>& NewRange, const TRange<float>& VisibleRange, const FIntPoint& AllottedSize, const FIntPoint& InDesiredSize)
+void FTrackEditorThumbnailCache::Update(const TRange<float>& NewRange, const TRange<float>& VisibleRange, const FIntPoint& AllottedSize, const FIntPoint& InDesiredSize, EThumbnailQuality Quality, double InCurrentTime)
 {
 	PreviousCache.TimeRange = CurrentCache.TimeRange;
 	PreviousCache.VisibleRange = CurrentCache.VisibleRange;
 	PreviousCache.AllottedSize = CurrentCache.AllottedSize;
 	PreviousCache.DesiredSize = CurrentCache.DesiredSize;
+	PreviousCache.Quality = CurrentCache.Quality;
 
 	CurrentCache.TimeRange = NewRange;
 	CurrentCache.VisibleRange = VisibleRange;
 	CurrentCache.AllottedSize = AllottedSize;
 	CurrentCache.DesiredSize = InDesiredSize;
+	CurrentCache.Quality = Quality;
 
-	Revalidate();
+	Revalidate(InCurrentTime);
 
-	// Only update the single reference frame value once we've updated, since that can get set at any time, but Update() many be throttled
+	// Only update the single reference frame value once we've updated, since that can get set at any time, but Update() may be throttled
 	PreviousCache.SingleReferenceFrame = CurrentCache.SingleReferenceFrame;
 }
 
@@ -285,19 +290,32 @@ FIntPoint FTrackEditorThumbnailCache::CalculateTextureSize() const
 
 	float SizeRatio = float(CurrentCache.DesiredSize.X) / CurrentCache.DesiredSize.Y;
 
-	FIntPoint NewSize = CurrentCache.DesiredSize;
+	float X = CurrentCache.DesiredSize.X;
+	float Y = CurrentCache.DesiredSize.Y;
+
 	if (SizeRatio > DesiredRatio)
 	{
 		// Take width
-		NewSize = FIntPoint(CurrentCache.DesiredSize.X, CurrentCache.DesiredSize.X / DesiredRatio);
+		Y = CurrentCache.DesiredSize.X / DesiredRatio;
 	}
 	else if (SizeRatio < DesiredRatio)
 	{
 		// Take height
-		NewSize = FIntPoint(CurrentCache.DesiredSize.Y * DesiredRatio, CurrentCache.DesiredSize.Y);
+		X = CurrentCache.DesiredSize.Y * DesiredRatio;
+	}
+
+	float Scale;
+	switch (CurrentCache.Quality)
+	{
+		case EThumbnailQuality::Draft: 	Scale = 0.5f; 	break;
+		case EThumbnailQuality::Best: 	Scale = 2.f; 	break;
+		default: 						Scale = 1.f; 	break;
 	}
 
-	return NewSize;
+	return FIntPoint(
+		FMath::RoundToInt(X * Scale),
+		FMath::RoundToInt(Y * Scale)
+		);
 }
 
 bool FTrackEditorThumbnailCache::ShouldRegenerateEverything() const
@@ -327,7 +345,7 @@ void FTrackEditorThumbnailCache::DrawViewportThumbnail(FTrackEditorThumbnail& Tr
 		// Finalize the view
 		InternalViewportClient->bLockedCameraView = true;
 		InternalViewportClient->UpdateViewForLockedActor();
-		GWorld->SendAllEndOfFrameUpdates();
+		InternalViewportClient->GetWorld()->SendAllEndOfFrameUpdates();
 
 		// Update the viewport RHI if necessary
 		FIntPoint Size = CalculateTextureSize();
@@ -357,9 +375,9 @@ void FTrackEditorThumbnailCache::DrawViewportThumbnail(FTrackEditorThumbnail& Tr
 	}
 }
 
-void FTrackEditorThumbnailCache::Revalidate()
+void FTrackEditorThumbnailCache::Revalidate(double InCurrentTime)
 {
-	if (CurrentCache == PreviousCache && !bForceRedraw)
+	if (CurrentCache == PreviousCache && !bForceRedraw && !bNeedsNewThumbnails)
 	{
 		return;
 	}
@@ -370,15 +388,32 @@ void FTrackEditorThumbnailCache::Revalidate()
 		ThumbnailPool.Pin()->RemoveThumbnailsNeedingRedraw(Thumbnails);
 		ThumbnailsNeedingRedraw.Reset();
 		Thumbnails.Reset();
+		bNeedsNewThumbnails = false;
 		return;
 	}
 
+	if (CurrentCache.Quality != PreviousCache.Quality)
+	{
+		SetupViewportEngineFlags();
+	}
+
+	bNeedsNewThumbnails = true;
+
 	if (ShouldRegenerateEverything())
 	{
 		ThumbnailPool.Pin()->RemoveThumbnailsNeedingRedraw(Thumbnails);
 		Thumbnails.Reset();
 	}
-	
+
+	if (InCurrentTime - LastComputationTime > 0.25f)
+	{
+		ComputeNewThumbnails();
+		LastComputationTime = InCurrentTime;
+	}
+}
+
+void FTrackEditorThumbnailCache::ComputeNewThumbnails()
+{
 	ThumbnailsNeedingRedraw.Reset();
 
 	if (CurrentCache.SingleReferenceFrame.IsSet())
@@ -403,6 +438,7 @@ void FTrackEditorThumbnailCache::Revalidate()
 	}
 
 	bForceRedraw = false;
+	bNeedsNewThumbnails = false;
 }
 
 void FTrackEditorThumbnailCache::UpdateSingleThumbnail()
@@ -572,11 +608,12 @@ void FTrackEditorThumbnailCache::Setup()
 		InternalViewportClient->ViewportType = LVT_Perspective;
 		InternalViewportClient->bDisableInput = true;
 		InternalViewportClient->bDrawAxes = false;
-		InternalViewportClient->EngineShowFlags = FEngineShowFlags(ESFIM_Game);
-		// @todo: DisableAdvancedFeatures()? Faster, but can yield grossly incorrect thumbnails
-		InternalViewportClient->EngineShowFlags.SetMotionBlur(false);
 		InternalViewportClient->SetAllowCinematicPreview(false);
 		InternalViewportClient->SetRealtime(false);
+
+		SetupViewportEngineFlags();
+
+		InternalViewportClient->ViewState.GetReference()->SetSequencerState(true);
 	}
 
 	if (!InternalViewportScene.IsValid())
@@ -584,4 +621,28 @@ void FTrackEditorThumbnailCache::Setup()
 		InternalViewportScene = MakeShareable(new FSceneViewport(InternalViewportClient.Get(), nullptr));
 		InternalViewportClient->Viewport = InternalViewportScene.Get();
 	}
+}
+
+void FTrackEditorThumbnailCache::SetupViewportEngineFlags()
+{
+	if (!InternalViewportClient.IsValid())
+	{
+		return;
+	}
+
+	InternalViewportClient->EngineShowFlags = FEngineShowFlags(ESFIM_Game);
+
+	switch (CurrentCache.Quality)
+	{
+	case EThumbnailQuality::Draft:
+		InternalViewportClient->EngineShowFlags.DisableAdvancedFeatures();
+		break;
+
+	case EThumbnailQuality::Normal:
+	case EThumbnailQuality::Best:
+		InternalViewportClient->EngineShowFlags.SetMotionBlur(false);
+		break;
+	}
+
+	InternalViewportClient->Invalidate();
 }
\ No newline at end of file
diff --git a/Engine/Source/Editor/MovieSceneTools/Private/TrackEditorThumbnail/TrackEditorThumbnail.h b/Engine/Source/Editor/MovieSceneTools/Private/TrackEditorThumbnail/TrackEditorThumbnail.h
--- a/Engine/Source/Editor/MovieSceneTools/Private/TrackEditorThumbnail/TrackEditorThumbnail.h
+++ b/Engine/Source/Editor/MovieSceneTools/Private/TrackEditorThumbnail/TrackEditorThumbnail.h
@@ -2,6 +2,7 @@
 
 #pragma once
 
+#include "MovieSceneToolsUserSettings.h"
 
 class FThumbnailSection;
 class FSlateRenderTargetRHI;
@@ -99,7 +100,7 @@ struct IThumbnailClient
 /** Cache data */
 struct FThumbnailCacheData
 {
-	FThumbnailCacheData() : VisibleRange(0.f), TimeRange(0.f), AllottedSize(0,0), DesiredSize(0, 0) {}
+	FThumbnailCacheData() : VisibleRange(0.f), TimeRange(0.f), AllottedSize(0,0), DesiredSize(0, 0), Quality(EThumbnailQuality::Normal) {}
 
 	bool operator==(const FThumbnailCacheData& RHS) const
 	{
@@ -108,6 +109,7 @@ struct FThumbnailCacheData
 			VisibleRange == RHS.VisibleRange &&
 			TimeRange == RHS.TimeRange &&
 			DesiredSize == RHS.DesiredSize &&
+			Quality == RHS.Quality &&
 			SingleReferenceFrame == RHS.SingleReferenceFrame;
 	}
 
@@ -118,6 +120,7 @@ struct FThumbnailCacheData
 			VisibleRange != RHS.VisibleRange ||
 			TimeRange != RHS.TimeRange ||
 			DesiredSize != RHS.DesiredSize ||
+			Quality != RHS.Quality ||
 			SingleReferenceFrame != RHS.SingleReferenceFrame;
 	}
 
@@ -129,6 +132,8 @@ struct FThumbnailCacheData
 	FIntPoint AllottedSize;
 	/** Desired frame size constraint */
 	FIntPoint DesiredSize;
+	/** Thumbnail quality */
+	EThumbnailQuality Quality;
 	/** Set when we want to render a single reference frame */
 	TOptional<float> SingleReferenceFrame;
 };
@@ -144,9 +149,9 @@ public:
 	void SetSingleReferenceFrame(TOptional<float> InReferenceFrame);
 	TOptional<float> GetSingleReferenceFrame() const { return CurrentCache.SingleReferenceFrame; }
 
-	void Update(const TRange<float>& NewRange, const TRange<float>& VisibleRange, const FIntPoint& AllottedSize, const FIntPoint& InDesiredSize);
+	void Update(const TRange<float>& NewRange, const TRange<float>& VisibleRange, const FIntPoint& AllottedSize, const FIntPoint& InDesiredSize, EThumbnailQuality InQuality, double InCurrentTime);
 
-	void Revalidate();
+	void Revalidate(double InCurrentTime);
 
 	void DrawViewportThumbnail(FTrackEditorThumbnail& TrackEditorThumbnail);
 
@@ -157,6 +162,8 @@ public:
 
 protected:
 
+	void ComputeNewThumbnails();
+
 	void Setup();
 
 	bool ShouldRegenerateEverything() const;
@@ -171,6 +178,8 @@ protected:
 
 	void GenerateBack(const TRange<float>& Boundary);
 
+	void SetupViewportEngineFlags();
+
 protected:
 
 	/** Thumbnail client used for paint notifications */
@@ -195,6 +204,9 @@ protected:
 	/** The number of frames we've rendered */
 	uint32 FrameCount;
 
+	double LastComputationTime;
+	bool bNeedsNewThumbnails;
+
 	/** Whether to force a redraw or not */
 	bool bForceRedraw;
 };
diff --git a/Engine/Source/Editor/MovieSceneTools/Private/TrackEditors/AudioTrackEditor.cpp b/Engine/Source/Editor/MovieSceneTools/Private/TrackEditors/AudioTrackEditor.cpp
--- a/Engine/Source/Editor/MovieSceneTools/Private/TrackEditors/AudioTrackEditor.cpp
+++ b/Engine/Source/Editor/MovieSceneTools/Private/TrackEditors/AudioTrackEditor.cpp
@@ -711,7 +711,8 @@ void FAudioSection::Tick( const FGeometry& AllottedGeometry, const FGeometry& Pa
 		if (!FMath::IsNearlyEqual(DrawRange.GetLowerBoundValue(), StoredDrawRange.GetLowerBoundValue()) ||
 			!FMath::IsNearlyEqual(DrawRange.GetUpperBoundValue(), StoredDrawRange.GetUpperBoundValue()) ||
 			XOffset != StoredXOffset || XSize != StoredXSize || Track->GetColorTint() != StoredColor ||
-			StoredSoundWave != SoundWave)
+			StoredSoundWave != SoundWave ||
+			StoredShowIntensity != AudioSection->ShouldShowIntensity())
 		{
 			float DisplayScale = XSize / DrawRange.Size<float>();
 
@@ -741,6 +742,7 @@ void FAudioSection::RegenerateWaveforms(TRange<float> DrawRange, int32 XOffset,
 	StoredXOffset = XOffset;
 	StoredXSize = XSize;
 	StoredColor = ColorTint;
+	StoredShowIntensity =  Cast<UMovieSceneAudioSection>(&Section)->ShouldShowIntensity();
 
 	if (DrawRange.IsDegenerate() || DrawRange.IsEmpty() || Cast<UMovieSceneAudioSection>(&Section)->GetSound() == NULL)
 	{
@@ -792,7 +794,7 @@ void CopyInterpSoundTrack(TSharedRef<ISequencer> Sequencer, UInterpTrackSound* M
 {
 	if (FMatineeImportTools::CopyInterpSoundTrack(MatineeSoundTrack, AudioTrack))
 	{
-		Sequencer.Get().NotifyMovieSceneDataChanged();
+		Sequencer.Get().NotifyMovieSceneDataChanged( EMovieSceneDataChangeType::MovieSceneStructureItemAdded );
 	}
 }
 
@@ -939,7 +941,7 @@ void FAudioTrackEditor::HandleAddAudioTrackMenuEntryExecute()
 
 	NewTrack->SetDisplayName(LOCTEXT("AudioTrackName", "Audio"));
 
-	GetSequencer()->NotifyMovieSceneDataChanged();
+	GetSequencer()->NotifyMovieSceneDataChanged( EMovieSceneDataChangeType::MovieSceneStructureItemAdded );
 }
 
 TSharedRef<SWidget> FAudioTrackEditor::BuildAudioSubMenu(UMovieSceneTrack* Track)
@@ -997,7 +999,7 @@ void FAudioTrackEditor::OnAudioAssetSelected(const FAssetData& AssetData, UMovie
 			float KeyTime = GetSequencer()->GetGlobalTime();
 			AudioTrack->AddNewSound( NewSound, KeyTime );
 
-			GetSequencer()->NotifyMovieSceneDataChanged();
+			GetSequencer()->NotifyMovieSceneDataChanged( EMovieSceneDataChangeType::MovieSceneStructureItemAdded );
 		}
 	}
 }
diff --git a/Engine/Source/Editor/MovieSceneTools/Private/TrackEditors/AudioTrackEditor.h b/Engine/Source/Editor/MovieSceneTools/Private/TrackEditors/AudioTrackEditor.h
--- a/Engine/Source/Editor/MovieSceneTools/Private/TrackEditors/AudioTrackEditor.h
+++ b/Engine/Source/Editor/MovieSceneTools/Private/TrackEditors/AudioTrackEditor.h
@@ -104,6 +104,7 @@ private:
 
 	/** Stored data about the waveform to determine when it is invalidated. */
 	TRange<float> StoredDrawRange;
+	bool StoredShowIntensity;
 	int32 StoredXOffset;
 	int32 StoredXSize;
 	FColor StoredColor;
diff --git a/Engine/Source/Editor/MovieSceneTools/Private/TrackEditors/CameraCutTrackEditor.cpp b/Engine/Source/Editor/MovieSceneTools/Private/TrackEditors/CameraCutTrackEditor.cpp
--- a/Engine/Source/Editor/MovieSceneTools/Private/TrackEditors/CameraCutTrackEditor.cpp
+++ b/Engine/Source/Editor/MovieSceneTools/Private/TrackEditors/CameraCutTrackEditor.cpp
@@ -166,6 +166,9 @@ UMovieSceneCameraCutTrack* FCameraCutTrackEditor::FindOrCreateCameraCutTrack()
 
 	if (CameraCutTrack == nullptr)
 	{
+		const FScopedTransaction Transaction(LOCTEXT("AddCameraCutTrack_Transaction", "Add Camera Cut Track"));
+		FocusedMovieScene->Modify();
+		
 		CameraCutTrack = FocusedMovieScene->AddCameraCutTrack(UMovieSceneCameraCutTrack::StaticClass());
 	}
 
@@ -216,7 +219,7 @@ bool FCameraCutTrackEditor::HandleAddCameraCutTrackMenuEntryCanExecute() const
 void FCameraCutTrackEditor::HandleAddCameraCutTrackMenuEntryExecute()
 {
 	FindOrCreateCameraCutTrack();
-	GetSequencer()->NotifyMovieSceneDataChanged();
+	GetSequencer()->NotifyMovieSceneDataChanged( EMovieSceneDataChangeType::MovieSceneStructureItemAdded );
 }
 
 bool FCameraCutTrackEditor::IsCameraPickable(const AActor* const PickableActor)
@@ -322,15 +325,13 @@ void FCameraCutTrackEditor::OnLockCameraClicked(ECheckBoxState CheckBoxState)
 	}
 	else
 	{
-		// TODO: MaxC, is this right?
 		GetSequencer()->UpdateCameraCut(nullptr, nullptr);
 		GetSequencer()->SetPerspectiveViewportCameraCutEnabled(false);
 	}
 
-	GetSequencer()->UpdateRuntimeInstances();
+	GetSequencer()->SetGlobalTime(GetSequencer()->GetGlobalTime());
 }
 
-
 FText FCameraCutTrackEditor::GetLockCameraToolTip() const
 {
 	return IsCameraLocked() == ECheckBoxState::Checked ?
diff --git a/Engine/Source/Editor/MovieSceneTools/Private/TrackEditors/CinematicShotTrackEditor.cpp b/Engine/Source/Editor/MovieSceneTools/Private/TrackEditors/CinematicShotTrackEditor.cpp
--- a/Engine/Source/Editor/MovieSceneTools/Private/TrackEditors/CinematicShotTrackEditor.cpp
+++ b/Engine/Source/Editor/MovieSceneTools/Private/TrackEditors/CinematicShotTrackEditor.cpp
@@ -190,7 +190,7 @@ void FCinematicShotTrackEditor::BuildTrackContextMenu( FMenuBuilder& MenuBuilder
 		FExecuteAction::CreateRaw(this, &FCinematicShotTrackEditor::ImportEDL )));
 
 	MenuBuilder.AddMenuEntry(
-		NSLOCTEXT( "Sequencer", "ExportEDL", "Export EDL" ),
+		NSLOCTEXT( "Sequencer", "ExportEDL", "Export EDL..." ),
 		NSLOCTEXT( "Sequencer", "ExportEDLTooltip", "Export Edit Decision List (EDL) for non-linear editors." ),
 		FSlateIcon(),
 		FUIAction(
@@ -257,10 +257,6 @@ UMovieSceneSubSection* FCinematicShotTrackEditor::CreateShotInternal(FString& Ne
 
 	// Create a cinematic shot section and set the time to be at the end of this shot. 
 	UMovieSceneSubSection* NewSection = CinematicShotTrack->AddSequence(NewSequence, NewShotStartTime, Duration, bInsertShot);	
-
-	// Update the track instances to generate the new section
-	GetSequencer()->UpdateRuntimeInstances();
-
 	return NewSection;
 }
 
@@ -276,7 +272,7 @@ void FCinematicShotTrackEditor::InsertShotAtCurrentTime()
 
 	CreateShotInternal(NewShotName, NewShotStartTime);
 
-	GetSequencer()->NotifyMovieSceneDataChanged();
+	GetSequencer()->NotifyMovieSceneDataChanged( EMovieSceneDataChangeType::MovieSceneStructureItemAdded );
 }
 
 
@@ -292,7 +288,7 @@ void FCinematicShotTrackEditor::InsertShot(UMovieSceneCinematicShotSection* Sect
 
 	CreateShotInternal(NewShotName, NewShotStartTime);
 
-	GetSequencer()->NotifyMovieSceneDataChanged();
+	GetSequencer()->NotifyMovieSceneDataChanged( EMovieSceneDataChangeType::MovieSceneStructureItemAdded );
 }
 
 
@@ -317,7 +313,7 @@ void FCinematicShotTrackEditor::InsertFillerAtCurrentTime()
 
 	NewCinematicShotSection->SetShotDisplayName(FText(LOCTEXT("Filler", "Filler")));
 
-	GetSequencer()->NotifyMovieSceneDataChanged();
+	GetSequencer()->NotifyMovieSceneDataChanged( EMovieSceneDataChangeType::MovieSceneStructureItemAdded );
 }
 
 void FCinematicShotTrackEditor::DuplicateShot(UMovieSceneCinematicShotSection* Section)
@@ -337,9 +333,9 @@ void FCinematicShotTrackEditor::DuplicateShot(UMovieSceneCinematicShotSection* S
 		NewShot->StartOffset = Section->StartOffset;
 		NewShot->TimeScale = Section->TimeScale;
 		NewShot->PrerollTime = Section->PrerollTime;
-	}
 
-	GetSequencer()->NotifyMovieSceneDataChanged();
+		GetSequencer()->NotifyMovieSceneDataChanged( EMovieSceneDataChangeType::MovieSceneStructureItemAdded );
+	}
 }
 
 void FCinematicShotTrackEditor::RenderShot(UMovieSceneCinematicShotSection* Section)
@@ -389,18 +385,15 @@ void FCinematicShotTrackEditor::NewTake(UMovieSceneCinematicShotSection* Section
 			UMovieSceneCinematicShotTrack* CinematicShotTrack = FindOrCreateCinematicShotTrack();
 			CinematicShotTrack->RemoveSection(*Section);
 
-			if (NewShot != nullptr)
-			{
-				NewShot->SetStartTime(NewShotStartTime);
-				NewShot->SetEndTime(NewShotEndTime);
-				NewShot->StartOffset = NewShotStartOffset;
-				NewShot->TimeScale = NewShotTimeScale;
-				NewShot->PrerollTime = NewShotPrerollTime;
-			}
+			NewShot->SetStartTime(NewShotStartTime);
+			NewShot->SetEndTime(NewShotEndTime);
+			NewShot->StartOffset = NewShotStartOffset;
+			NewShot->TimeScale = NewShotTimeScale;
+			NewShot->PrerollTime = NewShotPrerollTime;
+
+			GetSequencer()->NotifyMovieSceneDataChanged( EMovieSceneDataChangeType::MovieSceneStructureItemsChanged );
 		}
 	}
-
-	GetSequencer()->NotifyMovieSceneDataChanged();
 }
 
 
@@ -425,20 +418,21 @@ void FCinematicShotTrackEditor::SwitchTake(UMovieSceneCinematicShotSection* Sect
 		const bool bInsertSequence = false;
 		const float Duration = NewShotEndTime - NewShotStartTime;
 		UMovieSceneSubSection* NewShot = CinematicShotTrack->AddSequence(MovieSceneSequence, NewShotStartTime, Duration, bInsertSequence);	
-		
-		CinematicShotTrack->RemoveSection(*Section);
 
 		if (NewShot != nullptr)
 		{
+			CinematicShotTrack->RemoveSection(*Section);
+
 			NewShot->SetStartTime(NewShotStartTime);
 			NewShot->SetEndTime(NewShotEndTime);
 			NewShot->StartOffset = NewShotStartOffset;
 			NewShot->TimeScale = NewShotTimeScale;
 			NewShot->PrerollTime = NewShotPrerollTime;
+
+			GetSequencer()->NotifyMovieSceneDataChanged( EMovieSceneDataChangeType::MovieSceneStructureItemsChanged );
 		}
 	}
 
-	GetSequencer()->NotifyMovieSceneDataChanged();
 }
 
 
@@ -521,8 +515,6 @@ bool FCinematicShotTrackEditor::AddKeyInternal(float KeyTime, UMovieSceneSequenc
 		float EndTime = InMovieSceneSequence->GetMovieScene()->GetPlaybackRange().Size<float>();
 		CinematicShotTrack->AddSequence(InMovieSceneSequence, KeyTime, EndTime);
 
-		GetSequencer()->NotifyMovieSceneDataChanged();
-	
 		return true;
 	}
 
@@ -550,7 +542,7 @@ UMovieSceneCinematicShotTrack* FCinematicShotTrackEditor::FindOrCreateCinematicS
 	auto NewTrack = FocusedMovieScene->AddMasterTrack<UMovieSceneCinematicShotTrack>();
 	ensure(NewTrack);
 
-	GetSequencer()->NotifyMovieSceneDataChanged();
+	GetSequencer()->NotifyMovieSceneDataChanged( EMovieSceneDataChangeType::MovieSceneStructureItemAdded );
 
 	return NewTrack;
 }
@@ -591,7 +583,7 @@ void FCinematicShotTrackEditor::OnLockShotsClicked(ECheckBoxState CheckBoxState)
 		GetSequencer()->SetPerspectiveViewportCameraCutEnabled(false);
 	}
 
-	GetSequencer()->UpdateRuntimeInstances();
+	GetSequencer()->SetGlobalTime(GetSequencer()->GetGlobalTime());
 }
 
 
@@ -644,7 +636,6 @@ bool FCinematicShotTrackEditor::HandleSequenceAdded(float KeyTime, UMovieSceneSe
 	float Duration = Sequence->GetMovieScene()->GetPlaybackRange().Size<float>();
 	CinematicShotTrack->AddSequence(Sequence, KeyTime, Duration);
 
-	GetSequencer()->NotifyMovieSceneDataChanged();
 	return true;
 }
 
@@ -680,7 +671,7 @@ void FCinematicShotTrackEditor::ImportEDL()
 
 	if (MovieSceneToolHelpers::ShowImportEDLDialog(MovieScene, FrameRate, SaveDirectory))
 	{
-		GetSequencer()->NotifyMovieSceneDataChanged();
+		GetSequencer()->NotifyMovieSceneDataChanged( EMovieSceneDataChangeType::MovieSceneStructureItemsChanged );
 	}
 }
 
diff --git a/Engine/Source/Editor/MovieSceneTools/Private/TrackEditors/EventTrackEditor.cpp b/Engine/Source/Editor/MovieSceneTools/Private/TrackEditors/EventTrackEditor.cpp
--- a/Engine/Source/Editor/MovieSceneTools/Private/TrackEditors/EventTrackEditor.cpp
+++ b/Engine/Source/Editor/MovieSceneTools/Private/TrackEditors/EventTrackEditor.cpp
@@ -99,7 +99,7 @@ void FEventTrackEditor::HandleAddEventTrackMenuEntryExecute()
 	NewTrack->AddSection(*NewSection);
 	NewTrack->SetDisplayName(LOCTEXT("TrackName", "Events"));
 
-	GetSequencer()->NotifyMovieSceneDataChanged();
+	GetSequencer()->NotifyMovieSceneDataChanged( EMovieSceneDataChangeType::MovieSceneStructureItemAdded );
 }
 
 
diff --git a/Engine/Source/Editor/MovieSceneTools/Private/TrackEditors/FadeTrackEditor.cpp b/Engine/Source/Editor/MovieSceneTools/Private/TrackEditors/FadeTrackEditor.cpp
--- a/Engine/Source/Editor/MovieSceneTools/Private/TrackEditors/FadeTrackEditor.cpp
+++ b/Engine/Source/Editor/MovieSceneTools/Private/TrackEditors/FadeTrackEditor.cpp
@@ -91,7 +91,7 @@ void FFadeTrackEditor::HandleAddFadeTrackMenuEntryExecute()
 
 	FadeTrack->AddSection(*NewSection);
 
-	GetSequencer()->NotifyMovieSceneDataChanged();
+	GetSequencer()->NotifyMovieSceneDataChanged( EMovieSceneDataChangeType::MovieSceneStructureItemAdded );
 }
 
 
diff --git a/Engine/Source/Editor/MovieSceneTools/Private/TrackEditors/LevelVisibilityTrackEditor.cpp b/Engine/Source/Editor/MovieSceneTools/Private/TrackEditors/LevelVisibilityTrackEditor.cpp
--- a/Engine/Source/Editor/MovieSceneTools/Private/TrackEditors/LevelVisibilityTrackEditor.cpp
+++ b/Engine/Source/Editor/MovieSceneTools/Private/TrackEditors/LevelVisibilityTrackEditor.cpp
@@ -104,7 +104,7 @@ void FLevelVisibilityTrackEditor::OnAddTrack()
 
 	AddNewSection( FocusedMovieScene, NewTrack, ELevelVisibility::Visible );
 
-	GetSequencer()->NotifyMovieSceneDataChanged();
+	GetSequencer()->NotifyMovieSceneDataChanged( EMovieSceneDataChangeType::MovieSceneStructureItemAdded );
 }
 
 
@@ -140,7 +140,7 @@ void FLevelVisibilityTrackEditor::OnAddNewSection( UMovieSceneTrack* LevelVisibi
 	}
 
 	AddNewSection( FocusedMovieScene, LevelVisibilityTrack, Visibility );
-	GetSequencer()->NotifyMovieSceneDataChanged();
+	GetSequencer()->NotifyMovieSceneDataChanged( EMovieSceneDataChangeType::MovieSceneStructureItemAdded );
 }
 
 #undef LOCTEXT_NAMESPACE
diff --git a/Engine/Source/Editor/MovieSceneTools/Private/TrackEditors/MaterialTrackEditor.cpp b/Engine/Source/Editor/MovieSceneTools/Private/TrackEditors/MaterialTrackEditor.cpp
--- a/Engine/Source/Editor/MovieSceneTools/Private/TrackEditors/MaterialTrackEditor.cpp
+++ b/Engine/Source/Editor/MovieSceneTools/Private/TrackEditors/MaterialTrackEditor.cpp
@@ -133,7 +133,7 @@ void FMaterialTrackEditor::AddScalarParameter( FGuid ObjectBinding, UMovieSceneM
 		MaterialTrack->Modify();
 		MaterialTrack->AddScalarParameterKey(ParameterName, KeyTime, ParameterValue);
 	}
-	NotifyMovieSceneDataChanged();
+	GetSequencer()->NotifyMovieSceneDataChanged( EMovieSceneDataChangeType::MovieSceneStructureItemAdded );
 }
 
 
@@ -151,7 +151,7 @@ void FMaterialTrackEditor::AddColorParameter( FGuid ObjectBinding, UMovieSceneMa
 		MaterialTrack->Modify();
 		MaterialTrack->AddColorParameterKey( ParameterName, KeyTime, ParameterValue );
 	}
-	NotifyMovieSceneDataChanged();
+	GetSequencer()->NotifyMovieSceneDataChanged( EMovieSceneDataChangeType::MovieSceneStructureItemAdded );
 }
 
 
diff --git a/Engine/Source/Editor/MovieSceneTools/Private/TrackEditors/ParticleParameterTrackEditor.cpp b/Engine/Source/Editor/MovieSceneTools/Private/TrackEditors/ParticleParameterTrackEditor.cpp
--- a/Engine/Source/Editor/MovieSceneTools/Private/TrackEditors/ParticleParameterTrackEditor.cpp
+++ b/Engine/Source/Editor/MovieSceneTools/Private/TrackEditors/ParticleParameterTrackEditor.cpp
@@ -133,7 +133,7 @@ bool FParticleParameterTrackEditor::CanAddParticleParameterTrack( FGuid ObjectBi
 void FParticleParameterTrackEditor::AddParticleParameterTrack( FGuid ObjectBinding )
 {
 	FindOrCreateTrackForObject( ObjectBinding, UMovieSceneParticleParameterTrack::StaticClass(), TrackName, true);
-	NotifyMovieSceneDataChanged();
+	GetSequencer()->NotifyMovieSceneDataChanged( EMovieSceneDataChangeType::MovieSceneStructureItemAdded );
 }
 
 
@@ -157,7 +157,7 @@ void FParticleParameterTrackEditor::AddScalarParameter( FGuid ObjectBinding, UMo
 			ParticleParameterTrack->AddScalarParameterKey( ParameterName, KeyTime, Value );
 		}
 	}
-	NotifyMovieSceneDataChanged();
+	GetSequencer()->NotifyMovieSceneDataChanged( EMovieSceneDataChangeType::MovieSceneStructureItemAdded );
 }
 
 
@@ -181,7 +181,7 @@ void FParticleParameterTrackEditor::AddVectorParameter( FGuid ObjectBinding, UMo
 			ParticleParameterTrack->AddVectorParameterKey( ParameterName, KeyTime, Value );
 		}
 	}
-	NotifyMovieSceneDataChanged();
+	GetSequencer()->NotifyMovieSceneDataChanged( EMovieSceneDataChangeType::MovieSceneStructureItemAdded );
 }
 
 
@@ -205,7 +205,7 @@ void FParticleParameterTrackEditor::AddColorParameter( FGuid ObjectBinding, UMov
 			ParticleParameterTrack->AddColorParameterKey( ParameterName, KeyTime, Value );
 		}
 	}
-	NotifyMovieSceneDataChanged();
+	GetSequencer()->NotifyMovieSceneDataChanged( EMovieSceneDataChangeType::MovieSceneStructureItemAdded );
 }
 
 
diff --git a/Engine/Source/Editor/MovieSceneTools/Private/TrackEditors/ParticleTrackEditor.cpp b/Engine/Source/Editor/MovieSceneTools/Private/TrackEditors/ParticleTrackEditor.cpp
--- a/Engine/Source/Editor/MovieSceneTools/Private/TrackEditors/ParticleTrackEditor.cpp
+++ b/Engine/Source/Editor/MovieSceneTools/Private/TrackEditors/ParticleTrackEditor.cpp
@@ -345,7 +345,7 @@ void CopyInterpParticleTrack(TSharedRef<ISequencer> Sequencer, UInterpTrackToggl
 {
 	if (FMatineeImportTools::CopyInterpParticleTrack(MatineeToggleTrack, ParticleTrack))
 	{
-		Sequencer.Get().NotifyMovieSceneDataChanged();
+		Sequencer.Get().NotifyMovieSceneDataChanged( EMovieSceneDataChangeType::MovieSceneStructureItemAdded );
 	}
 }
 
diff --git a/Engine/Source/Editor/MovieSceneTools/Private/TrackEditors/PropertyTrackEditors/ColorPropertyTrackEditor.cpp b/Engine/Source/Editor/MovieSceneTools/Private/TrackEditors/PropertyTrackEditors/ColorPropertyTrackEditor.cpp
--- a/Engine/Source/Editor/MovieSceneTools/Private/TrackEditors/PropertyTrackEditors/ColorPropertyTrackEditor.cpp
+++ b/Engine/Source/Editor/MovieSceneTools/Private/TrackEditors/PropertyTrackEditors/ColorPropertyTrackEditor.cpp
@@ -5,6 +5,7 @@
 #include "ColorPropertyTrackEditor.h"
 #include "ColorPropertySection.h"
 #include "MatineeImportTools.h"
+#include "Matinee/InterpTrackLinearColorProp.h"
 #include "Matinee/InterpTrackColorProp.h"
 
 
@@ -71,17 +72,27 @@ void CopyInterpColorTrack(TSharedRef<ISequencer> Sequencer, UInterpTrackColorPro
 {
 	if (FMatineeImportTools::CopyInterpColorTrack(ColorPropTrack, ColorTrack))
 	{
-		Sequencer.Get().NotifyMovieSceneDataChanged();
+		Sequencer.Get().NotifyMovieSceneDataChanged( EMovieSceneDataChangeType::MovieSceneStructureItemAdded );
+	}
+}
+
+void CopyInterpLinearColorTrack(TSharedRef<ISequencer> Sequencer, UInterpTrackLinearColorProp* LinearColorPropTrack, UMovieSceneColorTrack* ColorTrack)
+{
+	if (FMatineeImportTools::CopyInterpLinearColorTrack(LinearColorPropTrack, ColorTrack))
+	{
+		Sequencer.Get().NotifyMovieSceneDataChanged( EMovieSceneDataChangeType::MovieSceneStructureItemAdded );
 	}
 }
 
 void FColorPropertyTrackEditor::BuildTrackContextMenu( FMenuBuilder& MenuBuilder, UMovieSceneTrack* Track )
 {
 	UInterpTrackColorProp* ColorPropTrack = nullptr;
+	UInterpTrackLinearColorProp* LinearColorPropTrack = nullptr;
 	for ( UObject* CopyPasteObject : GUnrealEd->MatineeCopyPasteBuffer )
 	{
 		ColorPropTrack = Cast<UInterpTrackColorProp>( CopyPasteObject );
-		if ( ColorPropTrack != nullptr )
+		LinearColorPropTrack = Cast<UInterpTrackLinearColorProp>( CopyPasteObject );
+		if ( ColorPropTrack != nullptr || LinearColorPropTrack != nullptr )
 		{
 			break;
 		}
@@ -92,7 +103,9 @@ void FColorPropertyTrackEditor::BuildTrackContextMenu( FMenuBuilder& MenuBuilder
 		NSLOCTEXT( "Sequencer", "PasteMatineeColorTrackTooltip", "Pastes keys from a Matinee color track into this track." ),
 		FSlateIcon(),
 		FUIAction(
-			FExecuteAction::CreateStatic( &CopyInterpColorTrack, GetSequencer().ToSharedRef(), ColorPropTrack, ColorTrack ),
-			FCanExecuteAction::CreateLambda( [=]()->bool { return ColorPropTrack != nullptr && ColorPropTrack->GetNumKeys() > 0 && ColorTrack != nullptr; } ) ) );
+			ColorPropTrack != nullptr ? 
+			FExecuteAction::CreateStatic( &CopyInterpColorTrack, GetSequencer().ToSharedRef(), ColorPropTrack, ColorTrack ) : 
+			FExecuteAction::CreateStatic( &CopyInterpLinearColorTrack, GetSequencer().ToSharedRef(), LinearColorPropTrack, ColorTrack ),			
+			FCanExecuteAction::CreateLambda( [=]()->bool { return ((ColorPropTrack != nullptr && ColorPropTrack->GetNumKeys() > 0) || (LinearColorPropTrack != nullptr && LinearColorPropTrack->GetNumKeys() > 0)) && ColorTrack != nullptr; } ) ) );
 }
 
diff --git a/Engine/Source/Editor/MovieSceneTools/Private/TrackEditors/PropertyTrackEditors/FloatPropertyTrackEditor.cpp b/Engine/Source/Editor/MovieSceneTools/Private/TrackEditors/PropertyTrackEditors/FloatPropertyTrackEditor.cpp
--- a/Engine/Source/Editor/MovieSceneTools/Private/TrackEditors/PropertyTrackEditors/FloatPropertyTrackEditor.cpp
+++ b/Engine/Source/Editor/MovieSceneTools/Private/TrackEditors/PropertyTrackEditors/FloatPropertyTrackEditor.cpp
@@ -29,7 +29,7 @@ void CopyInterpFloatTrack(TSharedRef<ISequencer> Sequencer, UInterpTrackFloatBas
 {
 	if (FMatineeImportTools::CopyInterpFloatTrack(MatineeFloatTrack, FloatTrack))
 	{
-		Sequencer.Get().NotifyMovieSceneDataChanged();
+		Sequencer.Get().NotifyMovieSceneDataChanged( EMovieSceneDataChangeType::MovieSceneStructureItemAdded );
 	}
 }
 
diff --git a/Engine/Source/Editor/MovieSceneTools/Private/TrackEditors/SpawnTrackEditor.cpp b/Engine/Source/Editor/MovieSceneTools/Private/TrackEditors/SpawnTrackEditor.cpp
--- a/Engine/Source/Editor/MovieSceneTools/Private/TrackEditors/SpawnTrackEditor.cpp
+++ b/Engine/Source/Editor/MovieSceneTools/Private/TrackEditors/SpawnTrackEditor.cpp
@@ -65,7 +65,7 @@ void FSpawnTrackEditor::HandleAddSpawnTrackMenuEntryExecute(FGuid ObjectBinding)
 {
 	FScopedTransaction AddSpawnTrackTransaction(LOCTEXT("AddSpawnTrack_Transaction", "Add Spawn Track"));
 	AddTrack(GetSequencer()->GetFocusedMovieSceneSequence()->GetMovieScene(), ObjectBinding, UMovieSceneSpawnTrack::StaticClass(), NAME_None);
-	NotifyMovieSceneDataChanged();
+	GetSequencer()->NotifyMovieSceneDataChanged( EMovieSceneDataChangeType::MovieSceneStructureItemAdded );
 }
 
 
diff --git a/Engine/Source/Editor/MovieSceneTools/Private/TrackEditors/SubTrackEditor.cpp b/Engine/Source/Editor/MovieSceneTools/Private/TrackEditors/SubTrackEditor.cpp
--- a/Engine/Source/Editor/MovieSceneTools/Private/TrackEditors/SubTrackEditor.cpp
+++ b/Engine/Source/Editor/MovieSceneTools/Private/TrackEditors/SubTrackEditor.cpp
@@ -69,14 +69,14 @@ public:
 		}
 		else if(UMovieSceneSubSection::GetRecordingSection() == &SectionObject)
 		{
-			const FString& ActorToRecord = UMovieSceneSubSection::GetActorToRecord();
+			AActor* ActorToRecord = UMovieSceneSubSection::GetActorToRecord();
 
 			ISequenceRecorder& SequenceRecorder = FModuleManager::LoadModuleChecked<ISequenceRecorder>("SequenceRecorder");
 			if(SequenceRecorder.IsRecording())
 			{
-				if(ActorToRecord.Len() > 0)
+				if(ActorToRecord != nullptr)
 				{
-					return FText::Format(LOCTEXT("RecordingIndicatorWithActor", "Sequence Recording for \"{0}\""), FText::FromString(UMovieSceneSubSection::GetActorToRecord()));
+					return FText::Format(LOCTEXT("RecordingIndicatorWithActor", "Sequence Recording for \"{0}\""), FText::FromString(ActorToRecord->GetActorLabel()));
 				}
 				else
 				{
@@ -85,9 +85,9 @@ public:
 			}
 			else
 			{
-				if(ActorToRecord.Len() > 0)
+				if(ActorToRecord != nullptr)
 				{
-					return FText::Format(LOCTEXT("RecordingPendingIndicatorWithActor", "Sequence Recording Pending for \"{0}\""), FText::FromString(UMovieSceneSubSection::GetActorToRecord()));
+					return FText::Format(LOCTEXT("RecordingPendingIndicatorWithActor", "Sequence Recording Pending for \"{0}\""), FText::FromString(ActorToRecord->GetActorLabel()));
 				}
 				else
 				{
@@ -324,7 +324,7 @@ TSharedPtr<SWidget> FSubTrackEditor::BuildOutlinerEditWidget(const FGuid& Object
 	.AutoWidth()
 	.VAlign(VAlign_Center)
 	[
-		FSequencerUtilities::MakeAddButton(LOCTEXT("SubText", "Sequence"), FOnGetContent::CreateSP(this, &FSubTrackEditor::HandleAddSubSequenceComboButtonGetMenuContent), Params.NodeIsHovered)
+		FSequencerUtilities::MakeAddButton(LOCTEXT("SubText", "Sequence"), FOnGetContent::CreateSP(this, &FSubTrackEditor::HandleAddSubSequenceComboButtonGetMenuContent, Track), Params.NodeIsHovered)
 	];
 }
 
@@ -433,7 +433,7 @@ void FSubTrackEditor::HandleAddSubTrackMenuEntryExecute()
 	auto NewTrack = FocusedMovieScene->AddMasterTrack<UMovieSceneSubTrack>();
 	ensure(NewTrack);
 
-	GetSequencer()->NotifyMovieSceneDataChanged();
+	GetSequencer()->NotifyMovieSceneDataChanged( EMovieSceneDataChangeType::MovieSceneStructureItemAdded );
 }
 
 /** Helper function - get the first PIE world (or first PIE client world if there is more than one) */
@@ -454,18 +454,19 @@ static UWorld* GetFirstPIEWorld()
 	return nullptr;
 }
 
-TSharedRef<SWidget> FSubTrackEditor::HandleAddSubSequenceComboButtonGetMenuContent()
+TSharedRef<SWidget> FSubTrackEditor::HandleAddSubSequenceComboButtonGetMenuContent(UMovieSceneTrack* InTrack)
 {
 	FMenuBuilder MenuBuilder(true, nullptr);
 
 	MenuBuilder.BeginSection(TEXT("RecordSequence"), LOCTEXT("RecordSequence", "Record Sequence"));
 	{
+		AActor* ActorToRecord = nullptr;
 		MenuBuilder.AddMenuEntry(
 			LOCTEXT("RecordNewSequence", "Record New Sequence"), 
 			LOCTEXT("RecordNewSequence_ToolTip", "Record a new level sequence into this sub-track from gameplay/simulation etc.\nThis only primes the track for recording. Click the record button to begin recording into this track once primed.\nOnly one sequence can be recorded at a time."), 
 			FSlateIcon(), 
 			FUIAction(
-				FExecuteAction::CreateSP(this, &FSubTrackEditor::HandleRecordNewSequence, FString()),
+				FExecuteAction::CreateSP(this, &FSubTrackEditor::HandleRecordNewSequence, ActorToRecord, InTrack),
 				FCanExecuteAction::CreateSP(this, &FSubTrackEditor::CanRecordNewSequence)));
 
 		if(UWorld* PIEWorld = GetFirstPIEWorld())
@@ -473,13 +474,13 @@ TSharedRef<SWidget> FSubTrackEditor::HandleAddSubSequenceComboButtonGetMenuConte
 			APlayerController* Controller = GEngine->GetFirstLocalPlayerController(PIEWorld);
 			if(Controller && Controller->GetPawn())
 			{
-				FString ActorNameToRecord = Controller->GetPawn()->GetFName().ToString();
+				ActorToRecord = Controller->GetPawn();
 				MenuBuilder.AddMenuEntry(
 					LOCTEXT("RecordNewSequenceFromPlayer", "Record New Sequence From Current Player"), 
 					LOCTEXT("RecordNewSequenceFromPlayer_ToolTip", "Record a new level sequence into this sub track using the current player's pawn.\nThis only primes the track for recording. Click the record button to begin recording into this track once primed.\nOnly one sequence can be recorded at a time."), 
 					FSlateIcon(), 
 					FUIAction(
-						FExecuteAction::CreateSP(this, &FSubTrackEditor::HandleRecordNewSequence, ActorNameToRecord),
+						FExecuteAction::CreateSP(this, &FSubTrackEditor::HandleRecordNewSequence, ActorToRecord, InTrack),
 						FCanExecuteAction::CreateSP(this, &FSubTrackEditor::CanRecordNewSequence)));
 			}
 		}
@@ -490,7 +491,7 @@ TSharedRef<SWidget> FSubTrackEditor::HandleAddSubSequenceComboButtonGetMenuConte
 	{
 		FAssetPickerConfig AssetPickerConfig;
 		{
-			AssetPickerConfig.OnAssetSelected = FOnAssetSelected::CreateRaw( this, &FSubTrackEditor::HandleAddSubSequenceComboButtonMenuEntryExecute);
+			AssetPickerConfig.OnAssetSelected = FOnAssetSelected::CreateRaw( this, &FSubTrackEditor::HandleAddSubSequenceComboButtonMenuEntryExecute, InTrack);
 			AssetPickerConfig.bAllowNullSelection = false;
 			AssetPickerConfig.InitialAssetViewType = EAssetViewType::Tile;
 			AssetPickerConfig.Filter.ClassNames.Add(TEXT("LevelSequence"));
@@ -512,7 +513,7 @@ TSharedRef<SWidget> FSubTrackEditor::HandleAddSubSequenceComboButtonGetMenuConte
 	return MenuBuilder.MakeWidget();
 }
 
-void FSubTrackEditor::HandleAddSubSequenceComboButtonMenuEntryExecute(const FAssetData& AssetData)
+void FSubTrackEditor::HandleAddSubSequenceComboButtonMenuEntryExecute(const FAssetData& AssetData, UMovieSceneTrack* InTrack)
 {
 	FSlateApplication::Get().DismissAllMenus();
 
@@ -522,15 +523,15 @@ void FSubTrackEditor::HandleAddSubSequenceComboButtonMenuEntryExecute(const FAss
 	{
 		UMovieSceneSequence* MovieSceneSequence = CastChecked<UMovieSceneSequence>(AssetData.GetAsset());
 
-		AnimatablePropertyChanged( FOnKeyProperty::CreateRaw( this, &FSubTrackEditor::AddKeyInternal, MovieSceneSequence) );
+		AnimatablePropertyChanged( FOnKeyProperty::CreateRaw( this, &FSubTrackEditor::AddKeyInternal, MovieSceneSequence, InTrack) );
 	}
 }
 
-bool FSubTrackEditor::AddKeyInternal(float KeyTime, UMovieSceneSequence* InMovieSceneSequence)
+bool FSubTrackEditor::AddKeyInternal(float KeyTime, UMovieSceneSequence* InMovieSceneSequence, UMovieSceneTrack* InTrack)
 {	
 	if (CanAddSubSequence(*InMovieSceneSequence))
 	{
-		auto SubTrack = FindOrCreateMasterTrack<UMovieSceneSubTrack>().Track;
+		UMovieSceneSubTrack* SubTrack = Cast<UMovieSceneSubTrack>(InTrack);
 		float Duration = InMovieSceneSequence->GetMovieScene()->GetPlaybackRange().Size<float>();
 		SubTrack->AddSequence(InMovieSceneSequence, KeyTime, Duration);
 
@@ -554,22 +555,22 @@ bool FSubTrackEditor::CanRecordNewSequence() const
 	return !UMovieSceneSubSection::IsSetAsRecording();
 }
 
-void FSubTrackEditor::HandleRecordNewSequence(FString NameOfActorToRecord)
+void FSubTrackEditor::HandleRecordNewSequence(AActor* InActorToRecord, UMovieSceneTrack* InTrack)
 {
 	FSlateApplication::Get().DismissAllMenus();
 
-	AnimatablePropertyChanged( FOnKeyProperty::CreateRaw( this, &FSubTrackEditor::HandleRecordNewSequenceInternal, NameOfActorToRecord) );
+	AnimatablePropertyChanged( FOnKeyProperty::CreateRaw( this, &FSubTrackEditor::HandleRecordNewSequenceInternal, InActorToRecord, InTrack) );
 }
 
-bool FSubTrackEditor::HandleRecordNewSequenceInternal(float KeyTime, FString NameOfActorToRecord)
+bool FSubTrackEditor::HandleRecordNewSequenceInternal(float KeyTime, AActor* InActorToRecord, UMovieSceneTrack* InTrack)
 {
-	auto SubTrack = FindOrCreateMasterTrack<UMovieSceneSubTrack>().Track;
+	UMovieSceneSubTrack* SubTrack = Cast<UMovieSceneSubTrack>(InTrack);
 	UMovieSceneSubSection* Section = SubTrack->AddSequenceToRecord();
 
 	// @todo: we could default to the same directory as a parent sequence, or the last sequence recorded. Lots of options!
 	Section->SetTargetSequenceName(GetDefault<USequenceRecorderSettings>()->SequenceName);
 	Section->SetTargetPathToRecordTo(GetDefault<USequenceRecorderSettings>()->SequenceRecordingBasePath.Path);
-	Section->SetNameOfActorToRecord(NameOfActorToRecord);
+	Section->SetActorToRecord(InActorToRecord);
 
 	return true;
 }
diff --git a/Engine/Source/Editor/MovieSceneTools/Private/TrackEditors/SubTrackEditor.h b/Engine/Source/Editor/MovieSceneTools/Private/TrackEditors/SubTrackEditor.h
--- a/Engine/Source/Editor/MovieSceneTools/Private/TrackEditors/SubTrackEditor.h
+++ b/Engine/Source/Editor/MovieSceneTools/Private/TrackEditors/SubTrackEditor.h
@@ -4,7 +4,7 @@
 
 
 class UMovieSceneSequence;
-
+class UMovieSceneTrack;
 
 /**
  * Tools for animatable property types such as floats ands vectors
@@ -62,13 +62,13 @@ private:
 	void HandleAddSubTrackMenuEntryExecute();
 
 	/** Callback for generating the menu of the "Add Sequence" combo button. */
-	TSharedRef<SWidget> HandleAddSubSequenceComboButtonGetMenuContent();
+	TSharedRef<SWidget> HandleAddSubSequenceComboButtonGetMenuContent(UMovieSceneTrack* InTrack);
 
 	/** Callback for executing a menu entry in the "Add Sequence" combo button. */
-	void HandleAddSubSequenceComboButtonMenuEntryExecute(const FAssetData& AssetData);
+	void HandleAddSubSequenceComboButtonMenuEntryExecute(const FAssetData& AssetData, UMovieSceneTrack* InTrack);
 
 	/** Delegate for AnimatablePropertyChanged in AddKey */
-	bool AddKeyInternal(float KeyTime, UMovieSceneSequence* InMovieSceneSequence);
+	bool AddKeyInternal(float KeyTime, UMovieSceneSequence* InMovieSceneSequence, UMovieSceneTrack* InTrack);
 
 	/** Callback for AnimatablePropertyChanged in HandleAssetAdded. */
 	bool HandleSequenceAdded(float KeyTime, UMovieSceneSequence* Sequence);
@@ -77,8 +77,8 @@ private:
 	bool CanRecordNewSequence() const;
 
 	/** Handle recording new sequence into a sub track */
-	void HandleRecordNewSequence(FString NameOfActorToRecord);
+	void HandleRecordNewSequence(AActor* InActorToRecord, UMovieSceneTrack* InTrack);
 
 	/** Actually handles the adding of the section */
-	bool HandleRecordNewSequenceInternal(float KeyTime, FString NameOfActorToRecord);
+	bool HandleRecordNewSequenceInternal(float KeyTime, AActor* InActorToRecord, UMovieSceneTrack* InTrack);
 };
diff --git a/Engine/Source/Editor/MovieSceneTools/Private/TrackEditors/TransformTrackEditor.cpp b/Engine/Source/Editor/MovieSceneTools/Private/TrackEditors/TransformTrackEditor.cpp
--- a/Engine/Source/Editor/MovieSceneTools/Private/TrackEditors/TransformTrackEditor.cpp
+++ b/Engine/Source/Editor/MovieSceneTools/Private/TrackEditors/TransformTrackEditor.cpp
@@ -243,7 +243,7 @@ void CopyInterpMoveTrack(TSharedRef<ISequencer> Sequencer, UInterpTrackMove* Mov
 {
 	if (FMatineeImportTools::CopyInterpMoveTrack(MoveTrack, TransformTrack))
 	{
-		Sequencer.Get().NotifyMovieSceneDataChanged();
+		Sequencer.Get().NotifyMovieSceneDataChanged( EMovieSceneDataChangeType::MovieSceneStructureItemAdded );
 	}
 }
 
@@ -961,7 +961,7 @@ void F3DTransformTrackEditor::ImportAnimSequenceTransforms(const FAssetData& Ass
 				Section->AddKey(TempKey.Time, FTransformKey(EKey3DTransformChannel::Scale, EAxis::Z, Scale.Z, bUnwindRotation), EMovieSceneKeyInterpolation::Linear);
 			}
 
-			Sequencer->NotifyMovieSceneDataChanged();
+			Sequencer->NotifyMovieSceneDataChanged( EMovieSceneDataChangeType::MovieSceneStructureItemAdded );
 		}
 	}
 }
diff --git a/Engine/Source/Editor/MovieSceneTools/Public/MatineeImportTools.h b/Engine/Source/Editor/MovieSceneTools/Public/MatineeImportTools.h
--- a/Engine/Source/Editor/MovieSceneTools/Public/MatineeImportTools.h
+++ b/Engine/Source/Editor/MovieSceneTools/Public/MatineeImportTools.h
@@ -4,6 +4,7 @@
 #include "Matinee/InterpTrackToggle.h"
 #include "MovieSceneParticleSection.h"
 
+class UInterpTrackLinearColorProp;
 class UInterpTrackColorProp;
 class UInterpTrackBoolProp;
 class UInterpTrackFloatBase;
@@ -58,6 +59,9 @@ public:
 	/** Copies keys from a matinee color track to a sequencer color track. */
 	static bool CopyInterpColorTrack( UInterpTrackColorProp* ColorPropTrack, UMovieSceneColorTrack* ColorTrack );
 
+	/** Copies keys from a matinee linear color track to a sequencer color track. */
+	static bool CopyInterpLinearColorTrack( UInterpTrackLinearColorProp* LinearColorPropTrack, UMovieSceneColorTrack* ColorTrack );
+
 	/** Copies keys from a matinee toggle track to a sequencer particle track. */
 	static bool CopyInterpParticleTrack( UInterpTrackToggle* MatineeToggleTrack, UMovieSceneParticleTrack* ParticleTrack );
 
diff --git a/Engine/Source/Editor/SceneOutliner/Private/ActorTreeItem.cpp b/Engine/Source/Editor/SceneOutliner/Private/ActorTreeItem.cpp
--- a/Engine/Source/Editor/SceneOutliner/Private/ActorTreeItem.cpp
+++ b/Engine/Source/Editor/SceneOutliner/Private/ActorTreeItem.cpp
@@ -240,7 +240,7 @@ FTreeItemPtr FActorTreeItem::CreateParent() const
 	}
 
 	AActor* ParentActor = ActorPtr->GetAttachParentActor();
-	if (ParentActor)
+	if (ParentActor && ensureMsgf(ParentActor != ActorPtr, TEXT("Encountered an Actor attached to itself (%s)"), *ParentActor->GetName()))
 	{
 		return MakeShareable(new FActorTreeItem(ParentActor));
 	}
diff --git a/Engine/Source/Editor/SequenceRecorder/Private/ActorRecording.cpp b/Engine/Source/Editor/SequenceRecorder/Private/ActorRecording.cpp
--- a/Engine/Source/Editor/SequenceRecorder/Private/ActorRecording.cpp
+++ b/Engine/Source/Editor/SequenceRecorder/Private/ActorRecording.cpp
@@ -68,7 +68,7 @@ bool UActorRecording::StartRecording(ULevelSequence* CurrentSequence, float Curr
 {
 	bNewComponentAddedWhileRecording = false;
 
-	if(ActorToRecord.IsValid())
+	if(GetActorToRecord() != nullptr)
 	{
 		if(CurrentSequence != nullptr)
 		{
@@ -77,7 +77,7 @@ bool UActorRecording::StartRecording(ULevelSequence* CurrentSequence, float Curr
 		else
 		{
 			TSharedPtr<FMovieSceneAnimationSectionRecorder> AnimationRecorder = MakeShareable(new FMovieSceneAnimationSectionRecorder(AnimationSettings, TargetAnimation.Get()));
-			AnimationRecorder->CreateSection(ActorToRecord.Get(), nullptr, FGuid(), 0.0f);
+			AnimationRecorder->CreateSection(GetActorToRecord(), nullptr, FGuid(), 0.0f);
 			AnimationRecorder->Record(0.0f);
 			SectionRecorders.Add(AnimationRecorder);			
 		}
@@ -111,9 +111,9 @@ void UActorRecording::GetSceneComponents(TArray<USceneComponent*>& OutArray, boo
 	// it is not enough to just go through the owned components array here
 	// we need to traverse the scene component hierarchy as well, as some components may be 
 	// owned by other actors (e.g. for pooling) and some may not be part of the hierarchy
-	if(ActorToRecord.IsValid())
+	if(GetActorToRecord() != nullptr)
 	{
-		USceneComponent* RootComponent = ActorToRecord->GetRootComponent();
+		USceneComponent* RootComponent = GetActorToRecord()->GetRootComponent();
 		if(RootComponent)
 		{
 			// note: GetChildrenComponents clears array!
@@ -122,7 +122,7 @@ void UActorRecording::GetSceneComponents(TArray<USceneComponent*>& OutArray, boo
 		}
 
 		// add owned components that are *not* part of the hierarchy
-		TInlineComponentArray<USceneComponent*> OwnedComponents(ActorToRecord.Get());
+		TInlineComponentArray<USceneComponent*> OwnedComponents(GetActorToRecord());
 		for(USceneComponent* OwnedComponent : OwnedComponents)
 		{
 			if(OwnedComponent->GetAttachParent() == nullptr && OwnedComponent != RootComponent)
@@ -133,7 +133,7 @@ void UActorRecording::GetSceneComponents(TArray<USceneComponent*>& OutArray, boo
 
 		if(!bIncludeNonCDO)
 		{
-			AActor* CDO = Cast<AActor>(ActorToRecord->GetClass()->GetDefaultObject());
+			AActor* CDO = Cast<AActor>(GetActorToRecord()->GetClass()->GetDefaultObject());
 
 			auto ShouldRemovePredicate = [&](UActorComponent* PossiblyRemovedComponent)
 				{
@@ -221,14 +221,14 @@ bool UActorRecording::ValidComponent(USceneComponent* SceneComponent) const
 
 void UActorRecording::FindOrAddFolder(UMovieScene* MovieScene)
 {
-	check(ActorToRecord.IsValid());
+	check(GetActorToRecord() != nullptr);
 
 	FName FolderName(NAME_None);
-	if(ActorToRecord.Get()->IsA<ACharacter>() || ActorToRecord.Get()->IsA<ASkeletalMeshActor>())
+	if(GetActorToRecord()->IsA<ACharacter>() || GetActorToRecord()->IsA<ASkeletalMeshActor>())
 	{
 		FolderName = TEXT("Characters");
 	}
-	else if(ActorToRecord.Get()->IsA<ACameraActor>() || ActorToRecord.Get()->IsA<ACameraRig_Crane>() || ActorToRecord.Get()->IsA<ACameraRig_Rail>())
+	else if(GetActorToRecord()->IsA<ACameraActor>() || GetActorToRecord()->IsA<ACameraRig_Crane>() || GetActorToRecord()->IsA<ACameraRig_Rail>())
 	{
 		FolderName = TEXT("Cameras");
 	}
@@ -265,12 +265,11 @@ void UActorRecording::StartRecordingActorProperties(ULevelSequence* CurrentSeque
 		// set up our spawnable for this actor
 		UMovieScene* MovieScene = CurrentSequence->GetMovieScene();
 
-		AActor* Actor = ActorToRecord.Get();
+		AActor* Actor = GetActorToRecord();
 		FString TemplateName = GetUniqueSpawnableName(MovieScene, Actor->GetName());
 
-		UClass* ActorClass = Actor->GetClass();
-		AActor* ObjectTemplate = NewObject<AActor>(MovieScene, ActorClass, *TemplateName);
-		
+		AActor* ObjectTemplate = CastChecked<AActor>(CurrentSequence->MakeSpawnableTemplateFromInstance(*Actor, *TemplateName));
+
 		if (ObjectTemplate)
 		{
 			TInlineComponentArray<USkeletalMeshComponent*> SkeletalMeshComponents;
@@ -303,7 +302,7 @@ void UActorRecording::StartRecordingActorProperties(ULevelSequence* CurrentSeque
 			const bool bIncludeNonCDO = false;
 			SyncTrackedComponents(bIncludeNonCDO);
 
-			TInlineComponentArray<USceneComponent*> SceneComponents(ActorToRecord.Get());
+			TInlineComponentArray<USceneComponent*> SceneComponents(GetActorToRecord());
 
 			// check if components need recording
 			TInlineComponentArray<USceneComponent*> ValidSceneComponents;
@@ -326,7 +325,7 @@ void UActorRecording::StartRecordingActorProperties(ULevelSequence* CurrentSeque
 			TSharedPtr<FMovieSceneAnimationSectionRecorder> FirstAnimRecorder = nullptr;
 			for(USceneComponent* SceneComponent : ValidSceneComponents)
 			{
-				TSharedPtr<FMovieSceneAnimationSectionRecorder> AnimRecorder = StartRecordingComponentProperties(SceneComponent->GetFName(), SceneComponent, ActorToRecord.Get(), CurrentSequence, CurrentSequenceTime);
+				TSharedPtr<FMovieSceneAnimationSectionRecorder> AnimRecorder = StartRecordingComponentProperties(SceneComponent->GetFName(), SceneComponent, GetActorToRecord(), CurrentSequence, CurrentSequenceTime);
 				if(!FirstAnimRecorder.IsValid() && AnimRecorder.IsValid())
 				{
 					FirstAnimRecorder = AnimRecorder;
@@ -334,7 +333,7 @@ void UActorRecording::StartRecordingActorProperties(ULevelSequence* CurrentSeque
 			}
 
 			// we need to create a transform track even if we arent recording transforms
-			if (FSequenceRecorder::Get().GetTransformRecorderFactory().CanRecordObject(ActorToRecord.Get()))
+			if (FSequenceRecorder::Get().GetTransformRecorderFactory().CanRecordObject(GetActorToRecord()))
 			{
 				UMovieScene3DTransformSectionRecorderSettings* Settings = ActorSettings.GetSettingsObject<UMovieScene3DTransformSectionRecorderSettings>();
 				check(Settings);
@@ -342,7 +341,7 @@ void UActorRecording::StartRecordingActorProperties(ULevelSequence* CurrentSeque
 				TSharedPtr<IMovieSceneSectionRecorder> Recorder = FSequenceRecorder::Get().GetTransformRecorderFactory().CreateSectionRecorder(Settings->bRecordTransforms, FirstAnimRecorder);
 				if(Recorder.IsValid())
 				{ 
-					Recorder->CreateSection(ActorToRecord.Get(), MovieScene, Guid, CurrentSequenceTime);
+					Recorder->CreateSection(GetActorToRecord(), MovieScene, Guid, CurrentSequenceTime);
 					Recorder->Record(CurrentSequenceTime);
 					SectionRecorders.Add(Recorder);
 				}
@@ -351,12 +350,12 @@ void UActorRecording::StartRecordingActorProperties(ULevelSequence* CurrentSeque
 			TArray<IMovieSceneSectionRecorderFactory*> ModularFeatures = IModularFeatures::Get().GetModularFeatureImplementations<IMovieSceneSectionRecorderFactory>(MovieSceneSectionRecorderFactoryName);
 			for (IMovieSceneSectionRecorderFactory* Factory : ModularFeatures)
 			{
-				if (Factory->CanRecordObject(ActorToRecord.Get()))
+				if (Factory->CanRecordObject(GetActorToRecord()))
 				{
 					TSharedPtr<IMovieSceneSectionRecorder> Recorder = Factory->CreateSectionRecorder(ActorSettings);
 					if (Recorder.IsValid())
 					{
-						Recorder->CreateSection(ActorToRecord.Get(), MovieScene, Guid, CurrentSequenceTime);
+						Recorder->CreateSection(GetActorToRecord(), MovieScene, Guid, CurrentSequenceTime);
 						Recorder->Record(CurrentSequenceTime);
 						SectionRecorders.Add(Recorder);
 					}
@@ -492,6 +491,23 @@ bool UActorRecording::IsRecording() const
 	return ActorToRecord.IsValid() && SectionRecorders.Num() > 0;
 }
 
+AActor* UActorRecording::GetActorToRecord() const
+{
+	if (ActorToRecord.IsValid())
+	{
+		AActor* OutActor = EditorUtilities::GetSimWorldCounterpartActor(ActorToRecord.Get());
+
+		if (OutActor != nullptr)
+		{
+			return OutActor;
+		}
+
+		return ActorToRecord.Get();
+	}
+
+	return nullptr;
+}
+
 static FName FindParentComponentOwnerClassName(USceneComponent* SceneComponent, UBlueprint* Blueprint)
 {
 	if(SceneComponent->GetAttachParent())
@@ -515,7 +531,7 @@ static FName FindParentComponentOwnerClassName(USceneComponent* SceneComponent,
 
 void UActorRecording::StartRecordingNewComponents(ULevelSequence* CurrentSequence, float CurrentSequenceTime)
 {
-	if (ActorToRecord.IsValid())
+	if (GetActorToRecord() != nullptr)
 	{
 		// find the new component(s)
 		TArray<USceneComponent*> NewComponents;
@@ -588,7 +604,7 @@ void UActorRecording::StartRecordingNewComponents(ULevelSequence* CurrentSequenc
 				NewName = SceneComponent->GetFName();
 			}
 
-			StartRecordingComponentProperties(NewName, SceneComponent, ActorToRecord.Get(), CurrentSequence, CurrentSequenceTime);
+			StartRecordingComponentProperties(NewName, SceneComponent, GetActorToRecord(), CurrentSequence, CurrentSequenceTime);
 
 			bNewComponentAddedWhileRecording = true;
 		}
diff --git a/Engine/Source/Editor/SequenceRecorder/Private/ActorRecording.h b/Engine/Source/Editor/SequenceRecorder/Private/ActorRecording.h
--- a/Engine/Source/Editor/SequenceRecorder/Private/ActorRecording.h
+++ b/Engine/Source/Editor/SequenceRecorder/Private/ActorRecording.h
@@ -38,6 +38,12 @@ public:
 		return Guid;
 	}
 
+	/** Get the actor to record. This finds the corresponding actor in the Simulation / PIE world. */
+	AActor* GetActorToRecord() const;
+
+	/** Set the actor to record */
+	void SetActorToRecord(AActor* InActor) { ActorToRecord = InActor; }
+
 private:
 	/** Check component validity for recording */
 	bool ValidComponent(USceneComponent* SceneComponent) const;
@@ -60,11 +66,8 @@ private:
 	/** Sync up tracked components with the actor */
 	void SyncTrackedComponents(bool bIncludeNonCDO = true);
 
-public:
-	/** The actor we want to record */
-	UPROPERTY(EditAnywhere, Category = "Actor Recording")
-	TLazyObjectPtr<AActor> ActorToRecord;
 
+public:
 	UPROPERTY(EditAnywhere, Category = "Actor Recording")
 	FActorRecordingSettings ActorSettings;
 
@@ -84,6 +87,10 @@ public:
 	bool bWasSpawnedPostRecord;
 
 private:
+	/** The actor we want to record */
+	UPROPERTY(EditAnywhere, Category = "Actor Recording")
+	TLazyObjectPtr<AActor> ActorToRecord;
+
 	/** This actor's current set of section recorders */
 	TArray<TSharedPtr<class IMovieSceneSectionRecorder>> SectionRecorders;
 
diff --git a/Engine/Source/Editor/SequenceRecorder/Private/AnimationRecorder.cpp b/Engine/Source/Editor/SequenceRecorder/Private/AnimationRecorder.cpp
--- a/Engine/Source/Editor/SequenceRecorder/Private/AnimationRecorder.cpp
+++ b/Engine/Source/Editor/SequenceRecorder/Private/AnimationRecorder.cpp
@@ -45,7 +45,7 @@ void FAnimationRecorder::SetSampleRateAndLength(float SampleRateHz, float Length
 	if (LengthInSeconds <= 0.f)
 	{
 		// invalid length passed in, default to unbounded
-		SampleRateHz = FAnimationRecordingSettings::DefaultSampleRate;
+		LengthInSeconds = FAnimationRecordingSettings::UnboundedMaximumLength;
 	}
 
 	IntervalTime = 1.0f / SampleRateHz;
@@ -311,6 +311,9 @@ UAnimSequence* FAnimationRecorder::StopRecord(bool bShowMessage)
 
 				FFloatCurve* FloatCurveData = nullptr;
 
+				TArray<float> TimesToRecord;
+				TArray<float> ValuesToRecord;
+
 				for (int32 FrameIndex = 0; FrameIndex < NumFrames; ++FrameIndex)
 				{
 					const float TimeToRecord = FrameIndex*IntervalTime;
@@ -325,13 +328,20 @@ UAnimSequence* FAnimationRecorder::StopRecord(bool bShowMessage)
 							FloatCurveData = static_cast<FFloatCurve*>(AnimationObject->RawCurveData.GetCurveData(UID, FRawCurveTracks::FloatType));
 						}
 					}
-					else if (FloatCurveData)
+
+					if (FloatCurveData)
 					{
-						FloatCurveData->FloatCurve.AddKey(TimeToRecord, CurCurve.Value);
+						TimesToRecord.Add(TimeToRecord);
+						ValuesToRecord.Add(CurCurve.Value);
 					}
 				}
-			}
-	
+
+				// Fill all the curve data at once
+				if (FloatCurveData)
+				{
+					FloatCurveData->FloatCurve.SetKeys(TimesToRecord, ValuesToRecord);
+				}
+			}	
 		}
 
 		//AnimationObject->RawCurveData.RemoveRedundantKeys();
diff --git a/Engine/Source/Editor/SequenceRecorder/Private/SClassPropertyRecorderSettings.cpp b/Engine/Source/Editor/SequenceRecorder/Private/SClassPropertyRecorderSettings.cpp
new file mode 100644
--- /dev/null
+++ b/Engine/Source/Editor/SequenceRecorder/Private/SClassPropertyRecorderSettings.cpp
@@ -0,0 +1,182 @@
+// Copyright 1998-2016 Epic Games, Inc. All Rights Reserved.
+
+#include "SequenceRecorderPrivatePCH.h"
+#include "SClassPropertyRecorderSettings.h"
+#include "EditorStyle.h"
+#include "SequenceRecorderSettings.h"
+#include "MovieSceneMultiPropertyRecorder.h"
+
+#define LOCTEXT_NAMESPACE "SClassPropertyRecorderSettings"
+
+void SClassPropertyRecorderSettings::Construct(const FArguments& Args, const TSharedRef<class IPropertyHandle>& InClassHandle, const TSharedRef<class IPropertyHandle>& InPropertiesHandle, IPropertyTypeCustomizationUtils& CustomizationUtils)
+{
+	PropertiesHandle = InPropertiesHandle;
+	ClassHandle = InClassHandle;
+
+	ChildSlot
+	[
+		SNew(SHorizontalBox)
+		+SHorizontalBox::Slot()
+		.FillWidth(1.0f)
+		.VAlign(VAlign_Center)
+		[
+			SNew(STextBlock)
+			.Text(this, &SClassPropertyRecorderSettings::GetText)
+			.Font(CustomizationUtils.GetRegularFont())
+			.AutoWrapText(true)
+		]
+		+SHorizontalBox::Slot()
+		.AutoWidth()
+		.VAlign(VAlign_Center)
+		[
+			SNew(SButton)
+			.ButtonStyle(FEditorStyle::Get(), "HoverHintOnly")
+			.ToolTipText(LOCTEXT("ChoosePropertiesButtonToolTip", "Choose properties to be recorded for this class"))
+			.OnClicked(this, &SClassPropertyRecorderSettings::HandleChoosePropertiesButtonClicked)
+			[
+				SNew(SImage)
+				.Image(FEditorStyle::GetBrush("PropertyWindow.Button_Ellipsis"))
+			]
+		]
+	];
+}
+
+FText SClassPropertyRecorderSettings::GetText() const
+{
+	if (PropertiesHandle.IsValid())
+	{
+		TArray<FName>& PropertiesToRecord = GetPropertyNameArray();
+		if (PropertiesToRecord.Num() > 0)
+		{
+			FText Text = FText::Format(LOCTEXT("SinglePropertyToRecordFormat", "{0}"), FText::FromName(PropertiesToRecord[0]));
+			for (int32 PropertyIndex = 1; PropertyIndex < PropertiesToRecord.Num(); ++PropertyIndex)
+			{
+				Text = FText::Format(LOCTEXT("PropertiesToRecordFormat", "{0}, {1}"), Text, FText::FromName(PropertiesToRecord[PropertyIndex]));
+			}
+			return Text;
+		}
+	}
+
+	return LOCTEXT("NoPropertiesToRecord", "None");
+}
+
+FReply SClassPropertyRecorderSettings::HandleChoosePropertiesButtonClicked()
+{
+	if (ClassHandle.IsValid() && PropertiesHandle.IsValid())
+	{
+		FDetailsViewArgs DetailsViewArgs(false, false, true, FDetailsViewArgs::HideNameArea, true);
+		FPropertyEditorModule& PropertyEditorModule = FModuleManager::GetModuleChecked<FPropertyEditorModule>("PropertyEditor");
+		TSharedRef<IDetailsView> DetailsView = PropertyEditorModule.CreateDetailView(DetailsViewArgs);
+
+		DetailsView->SetIsPropertyVisibleDelegate(FIsPropertyVisible::CreateSP(this, &SClassPropertyRecorderSettings::ShouldShowProperty));
+		DetailsView->SetIsPropertyReadOnlyDelegate(FIsPropertyReadOnly::CreateSP(this, &SClassPropertyRecorderSettings::IsPropertyReadOnly));
+		DetailsView->SetExtensionHandler(SharedThis(this));
+		DetailsView->SetDisableCustomDetailLayouts(true);
+
+		TArray<void*> RawData;
+		ClassHandle->AccessRawData(RawData);
+		check(RawData.Num() > 0);
+		TSubclassOf<UObject>& Class = *reinterpret_cast<TSubclassOf<UObject>*>(RawData[0]);
+		DetailsView->SetObject(Class.GetDefaultObject());
+
+		TSharedRef<SWindow> Window = SNew(SWindow)
+			.Title(LOCTEXT("PropertyPickerWindowTitle", "Choose Properties to Be Recorded"))
+			.ClientSize(FVector2D(400, 550));
+
+		Window->SetContent(
+			SNew(SBorder)
+			.BorderImage(FEditorStyle::GetBrush(TEXT("PropertyWindow.WindowBorder")))
+			[
+				DetailsView
+			]
+		);
+
+		// Parent the window to the root window
+		FSlateApplication::Get().AddModalWindow(Window, FGlobalTabmanager::Get()->GetRootWindow().ToSharedRef());
+	}
+	return FReply::Handled();
+}
+
+bool ShouldShowProperty_Recursive(const UProperty* InProperty)
+{
+	if (InProperty->IsA<UStructProperty>())
+	{
+		const UStructProperty* StructProperty = CastChecked<const UStructProperty>(InProperty);
+		for (const UProperty* Property = StructProperty->Struct->PropertyLink; Property; Property = Property->PropertyLinkNext)
+		{
+			if (Property->HasAnyPropertyFlags(CPF_Interp) && FMovieSceneMultiPropertyRecorder::CanPropertyBeRecorded(*Property))
+			{
+				return true;
+			}
+			else if(ShouldShowProperty_Recursive(Property))
+			{
+				return true;
+			}
+		}
+	}
+
+	return false;
+}
+
+bool SClassPropertyRecorderSettings::ShouldShowProperty(const FPropertyAndParent& InPropertyAndParent) const
+{
+	bool bShow = InPropertyAndParent.Property.HasAnyPropertyFlags(CPF_Interp) && FMovieSceneMultiPropertyRecorder::CanPropertyBeRecorded(InPropertyAndParent.Property);
+
+	// we also need to recurse into sub-properties to see if we should show them too
+	return bShow || ShouldShowProperty_Recursive(&InPropertyAndParent.Property);
+}
+
+bool SClassPropertyRecorderSettings::IsPropertyReadOnly(const FPropertyAndParent& InPropertyAndParent) const
+{
+	return true;
+}
+
+bool SClassPropertyRecorderSettings::IsPropertyExtendable(const UClass* InObjectClass, const class IPropertyHandle& PropertyHandle) const
+{
+	return PropertyHandle.GetProperty()->HasAnyPropertyFlags(CPF_Interp) && FMovieSceneMultiPropertyRecorder::CanPropertyBeRecorded(*PropertyHandle.GetProperty());
+}
+
+TSharedRef<SWidget> SClassPropertyRecorderSettings::GenerateExtensionWidget(const UClass* InObjectClass, TSharedPtr<IPropertyHandle> PropertyHandle)
+{
+	ECheckBoxState InitialState = ECheckBoxState::Unchecked;
+	FName PropertyName = *PropertyHandle->GeneratePathToProperty();
+
+	TArray<FName>& PropertiesToRecord = GetPropertyNameArray();
+	InitialState = PropertiesToRecord.Contains(PropertyName) ? ECheckBoxState::Checked : ECheckBoxState::Unchecked;
+
+	return SNew(SCheckBox)
+		.OnCheckStateChanged(this, &SClassPropertyRecorderSettings::HandlePropertyCheckStateChanged, PropertyHandle)
+		.IsChecked(InitialState);
+}
+
+void SClassPropertyRecorderSettings::HandlePropertyCheckStateChanged(ECheckBoxState InState, TSharedPtr<IPropertyHandle> PropertyHandle)
+{
+	if (PropertyHandle.IsValid() && PropertiesHandle.IsValid())
+	{
+		TArray<FName>& PropertiesToRecord = GetPropertyNameArray();
+
+		FName PropertyPath = *PropertyHandle->GeneratePathToProperty();
+		if (InState == ECheckBoxState::Checked)
+		{
+			PropertiesToRecord.AddUnique(PropertyPath);
+		}
+		else
+		{
+			PropertiesToRecord.Remove(PropertyPath);
+		}
+
+		PropertiesHandle->NotifyPostChange();
+	}
+}
+
+TArray<FName>& SClassPropertyRecorderSettings::GetPropertyNameArray() const
+{
+	check(PropertiesHandle.IsValid());
+
+	TArray<void*> RawData;
+	PropertiesHandle->AccessRawData(RawData);
+	check(RawData.Num() > 0);
+	return *reinterpret_cast<TArray<FName>*>(RawData[0]);
+}
+
+#undef LOCTEXT_NAMESPACE
\ No newline at end of file
diff --git a/Engine/Source/Editor/SequenceRecorder/Private/SSequenceRecorder.cpp b/Engine/Source/Editor/SequenceRecorder/Private/SSequenceRecorder.cpp
--- a/Engine/Source/Editor/SequenceRecorder/Private/SSequenceRecorder.cpp
+++ b/Engine/Source/Editor/SequenceRecorder/Private/SSequenceRecorder.cpp
@@ -66,9 +66,9 @@ private:
 	FText GetRecordingActorName() const
 	{
 		FText ActorName(LOCTEXT("InvalidActorName", "None"));
-		if (RecordingPtr.IsValid() && RecordingPtr.Get()->ActorToRecord.IsValid())
+		if (RecordingPtr.IsValid() && RecordingPtr.Get()->GetActorToRecord() != nullptr)
 		{
-			ActorName = FText::FromString(RecordingPtr.Get()->ActorToRecord.Get()->GetActorLabel());
+			ActorName = FText::FromString(RecordingPtr.Get()->GetActorToRecord()->GetActorLabel());
 		}
 
 		return ActorName;
