Fix Non-Unity CIS issues from Framework copy up.

diff --git a/Engine/Source/Editor/SequenceRecorder/Private/Sections/MovieScene3DTransformSectionRecorder.cpp b/Engine/Source/Editor/SequenceRecorder/Private/Sections/MovieScene3DTransformSectionRecorder.cpp
--- a/Engine/Source/Editor/SequenceRecorder/Private/Sections/MovieScene3DTransformSectionRecorder.cpp
+++ b/Engine/Source/Editor/SequenceRecorder/Private/Sections/MovieScene3DTransformSectionRecorder.cpp
@@ -7,6 +7,7 @@
 #include "Sections/MovieScene3DTransformSection.h"
 #include "Tracks/MovieScene3DTransformTrack.h"
 #include "SequenceRecorder.h"
+#include "SequenceRecorderSettings.h"
 
 TSharedPtr<IMovieSceneSectionRecorder> FMovieScene3DTransformSectionRecorderFactory::CreateSectionRecorder(const FActorRecordingSettings& InActorRecordingSettings) const
 {
diff --git a/Engine/Source/Editor/SequenceRecorder/Private/Sections/MovieScenePropertyRecorder.h b/Engine/Source/Editor/SequenceRecorder/Private/Sections/MovieScenePropertyRecorder.h
--- a/Engine/Source/Editor/SequenceRecorder/Private/Sections/MovieScenePropertyRecorder.h
+++ b/Engine/Source/Editor/SequenceRecorder/Private/Sections/MovieScenePropertyRecorder.h
@@ -5,6 +5,7 @@
 #include "CoreMinimal.h"
 #include "MovieSceneCommonHelpers.h"
 #include "MovieSceneSection.h"
+#include "SequenceRecorderSettings.h"
 
 /** Interface for a generic property recorder */
 class IMovieScenePropertyRecorder
diff --git a/Engine/Source/Editor/SkeletonEditor/Private/BoneProxy.h b/Engine/Source/Editor/SkeletonEditor/Private/BoneProxy.h
--- a/Engine/Source/Editor/SkeletonEditor/Private/BoneProxy.h
+++ b/Engine/Source/Editor/SkeletonEditor/Private/BoneProxy.h
@@ -2,6 +2,7 @@
 
 #pragma once
 
+#include "TickableEditorObject.h"
 #include "BoneProxy.generated.h"
 
 class UDebugSkelMeshComponent;
