Copying //UE4/Dev-Sequencer to //UE4/Dev-Main (Source: //UE4/Dev-Sequencer @ 3946692)

diff --git a/Engine/Source/Editor/MovieSceneTools/Public/MovieSceneToolHelpers.h b/Engine/Source/Editor/MovieSceneTools/Public/MovieSceneToolHelpers.h
--- a/Engine/Source/Editor/MovieSceneTools/Public/MovieSceneToolHelpers.h
+++ b/Engine/Source/Editor/MovieSceneTools/Public/MovieSceneToolHelpers.h
@@ -19,6 +19,9 @@ class UInterpTrackMoveAxis;
 struct FMovieSceneObjectBindingID;
 class UMovieSceneTrack;
 
+struct FMovieSceneFloatValue;
+template<typename ChannelType> struct TMovieSceneChannel;
+
 DECLARE_DELEGATE_TwoParams(FOnEnumSelectionChanged, int32 /*Selection*/, ESelectInfo::Type /*SelectionType*/);
 
 class MOVIESCENETOOLS_API MovieSceneToolHelpers
@@ -32,7 +35,7 @@ public:
 	 * @param Time	The time at which to trim
 	 * @param bTrimLeft Trim left or trim right
 	 */
-	static void TrimSection(const TSet<TWeakObjectPtr<UMovieSceneSection>>& Sections, float Time, bool bTrimLeft);
+	static void TrimSection(const TSet<TWeakObjectPtr<UMovieSceneSection>>& Sections, FFrameNumber Time, bool bTrimLeft);
 
 	/**
 	 * Splits sections at the given time
@@ -40,7 +43,7 @@ public:
 	 * @param Sections The sections to split
 	 * @param Time	The time at which to split
 	 */
-	static void SplitSection(const TSet<TWeakObjectPtr<UMovieSceneSection>>& Sections, float Time);
+	static void SplitSection(const TSet<TWeakObjectPtr<UMovieSceneSection>>& Sections, FFrameNumber Time);
 
 	/**
 	 * Parse a shot name into its components.
@@ -79,7 +82,7 @@ public:
 	 * @param Time The time to generate the new shot name at
 	 * @return The new shot name
 	 */
-	static FString GenerateNewShotName(const TArray<UMovieSceneSection*>& AllSections, float Time);
+	static FString GenerateNewShotName(const TArray<UMovieSceneSection*>& AllSections, FFrameNumber Time);
 
 	/**
 	 * Gather takes - level sequence assets that have the same shot prefix and shot number in the same asset path (directory)
@@ -128,7 +131,7 @@ public:
 	 * @param InOpenDirectory Optional directory path to open from. If none given, a dialog will pop up to prompt the user
 	 * @return Whether the import was successful
 	 */
-	static bool ShowImportEDLDialog(UMovieScene* InMovieScene, float InFrameRate, FString InOpenDirectory = TEXT(""));
+	static bool ShowImportEDLDialog(UMovieScene* InMovieScene, FFrameRate InFrameRate, FString InOpenDirectory = TEXT(""));
 
 	/**
 	 * Show Export EDL Dialog
@@ -139,7 +142,7 @@ public:
 	 * @param InHandleFrames The number of handle frames to include for each shot.
 	 * @return Whether the export was successful
 	 */
-	static bool ShowExportEDLDialog(const UMovieScene* InMovieScene, float InFrameRate, FString InSaveDirectory = TEXT(""), int32 InHandleFrames = 8);
+	static bool ShowExportEDLDialog(const UMovieScene* InMovieScene, FFrameRate InFrameRate, FString InSaveDirectory = TEXT(""), int32 InHandleFrames = 8);
 
 	/**
 	 * Import FBX
@@ -159,12 +162,13 @@ public:
 	static EInterpCurveMode RichCurveInterpolationToMatineeInterpolation( ERichCurveInterpMode InterpMode );
 
 	/*
-	 * Copy rich curve to move axis
+	 * Copy key data to move axis
 	 *
-	 * @param RichCurve The rich curve to copy from
+	 * @param KeyData The key data to copy from
 	 * @param MoveAxis The move axis to copy to
+	 * @param FrameRate The frame rate of the source channel
 	 */
-	static void CopyRichCurveToMoveAxis(const FRichCurve& RichCurve, UInterpTrackMoveAxis* MoveAxis);
+	static void CopyKeyDataToMoveAxis(const TMovieSceneChannel<FMovieSceneFloatValue>& KeyData, UInterpTrackMoveAxis* MoveAxis, FFrameRate FrameRate);
 
 	/*
 	 * Export the object binding to a camera anim
diff --git a/Engine/Source/Editor/MovieSceneTools/Public/PropertyTrackEditor.h b/Engine/Source/Editor/MovieSceneTools/Public/PropertyTrackEditor.h
--- a/Engine/Source/Editor/MovieSceneTools/Public/PropertyTrackEditor.h
+++ b/Engine/Source/Editor/MovieSceneTools/Public/PropertyTrackEditor.h
@@ -19,9 +19,9 @@
 /**
 * Tools for animatable property types such as floats ands vectors
 */
-template<typename TrackType, typename SectionType, typename KeyDataType>
+template<typename TrackType>
 class FPropertyTrackEditor
-	: public FKeyframeTrackEditor<TrackType, SectionType, KeyDataType>
+	: public FKeyframeTrackEditor<TrackType>
 {
 public:
 	/**
@@ -30,7 +30,7 @@ public:
 	 * @param InSequencer The sequencer instance to be used by this tool
 	 */
 	FPropertyTrackEditor( TSharedRef<ISequencer> InSequencer )
-		: FKeyframeTrackEditor<TrackType, SectionType, KeyDataType>( InSequencer )
+		: FKeyframeTrackEditor<TrackType>( InSequencer )
 	{ }
 
 	/**
@@ -40,7 +40,7 @@ public:
 	 * @param InWatchedPropertyTypes A list of property types that this editor can animate
 	 */
 	FPropertyTrackEditor( TSharedRef<ISequencer> InSequencer, TArrayView<const FAnimatedPropertyKey> InWatchedPropertyTypes )
-		: FKeyframeTrackEditor<TrackType, SectionType, KeyDataType>( InSequencer )
+		: FKeyframeTrackEditor<TrackType>( InSequencer )
 	{
 		for (const FAnimatedPropertyKey& Key : InWatchedPropertyTypes)
 		{
@@ -48,40 +48,6 @@ public:
 		}
 	}
 
-	DEPRECATED(4.16, "Please use FPropertyTrackEditor(TSharedRef<ISequencer>, TArrayView<const FAnimatedPropertyKey>)")
-	FPropertyTrackEditor( TSharedRef<ISequencer> InSequencer, FName WatchedPropertyTypeName )
-		: FKeyframeTrackEditor<TrackType, SectionType, KeyDataType>( InSequencer )
-	{
-		AddWatchedPropertyType( WatchedPropertyTypeName );
-	}
-
-	DEPRECATED(4.16, "Please use FPropertyTrackEditor(TSharedRef<ISequencer>, TArrayView<const FAnimatedPropertyKey>)")
-	FPropertyTrackEditor( TSharedRef<ISequencer> InSequencer, FName WatchedPropertyTypeName1, FName WatchedPropertyTypeName2 )
-		: FKeyframeTrackEditor<TrackType, SectionType, KeyDataType>( InSequencer )
-	{
-		AddWatchedPropertyType( WatchedPropertyTypeName1 );
-		AddWatchedPropertyType( WatchedPropertyTypeName2 );
-	}
-
-	DEPRECATED(4.16, "Please use FPropertyTrackEditor(TSharedRef<ISequencer>, TArrayView<const FAnimatedPropertyKey>)")
-	FPropertyTrackEditor( TSharedRef<ISequencer> InSequencer, FName WatchedPropertyTypeName1, FName WatchedPropertyTypeName2, FName WatchedPropertyTypeName3 )
-		: FKeyframeTrackEditor<TrackType, SectionType, KeyDataType>( InSequencer )
-	{
-		AddWatchedPropertyType( WatchedPropertyTypeName1 );
-		AddWatchedPropertyType( WatchedPropertyTypeName2 );
-		AddWatchedPropertyType( WatchedPropertyTypeName3 );
-	}
-
-	DEPRECATED(4.16, "Please use FPropertyTrackEditor(TSharedRef<ISequencer>, TArrayView<const FAnimatedPropertyKey>)")
-	FPropertyTrackEditor( TSharedRef<ISequencer> InSequencer, TArray<FName> InWatchedPropertyTypeNames )
-		: FKeyframeTrackEditor<TrackType, SectionType, KeyDataType>( InSequencer )
-	{
-		for ( FName WatchedPropertyTypeName : InWatchedPropertyTypeNames )
-		{
-			AddWatchedPropertyType( WatchedPropertyTypeName );
-		}
-	}
-
 	~FPropertyTrackEditor()
 	{
 		TSharedPtr<ISequencer> SequencerPtr = FMovieSceneTrackEditor::GetSequencer();
@@ -115,10 +81,9 @@ protected:
 	* Generates keys based on the new value from the property property change parameters.
 	*
 	* @param PropertyChangedParams Parameters associated with the property change.
-	* @param NewGeneratedKeys New keys which should be added due to the property change.
-	* @param DefaultGeneratedKeys Default value keys which should not be added, but may be needed for setting up defaults on new multi-channel tracks.
+	* @param OutGeneratedKeys Array of keys that are generated from the changed property
 	*/
-	virtual void GenerateKeysFromPropertyChanged( const FPropertyChangedParams& PropertyChangedParams, TArray<KeyDataType>& NewGeneratedKeys, TArray<KeyDataType>& DefaultGeneratedKeys ) = 0;
+	virtual void GenerateKeysFromPropertyChanged( const FPropertyChangedParams& PropertyChangedParams, FGeneratedTrackKeys& OutGeneratedKeys ) = 0;
 
 	/** When true, this track editor will only be used on properties which have specified it as a custom track class. This is necessary to prevent duplicate
 		property change handling in cases where a custom track editor handles the same type of data as one of the standard track editors. */
@@ -198,7 +163,7 @@ protected:
 	/** Adds a callback for property changes for the supplied property type name. */
 	void AddWatchedProperty( FAnimatedPropertyKey PropertyKey )
 	{
-		FMovieSceneTrackEditor::GetSequencer()->GetObjectChangeListener().GetOnAnimatablePropertyChanged( PropertyKey ).AddRaw( this, &FPropertyTrackEditor<TrackType, SectionType, KeyDataType>::OnAnimatedPropertyChanged );
+		FMovieSceneTrackEditor::GetSequencer()->GetObjectChangeListener().GetOnAnimatablePropertyChanged( PropertyKey ).AddRaw( this, &FPropertyTrackEditor::OnAnimatedPropertyChanged );
 		WatchedProperties.Add( PropertyKey );
 	}
 
@@ -239,13 +204,12 @@ private:
 	}
 
 	/** Adds a key based on a property change. */
-	FKeyPropertyResult OnKeyProperty( float KeyTime, FPropertyChangedParams PropertyChangedParams )
+	FKeyPropertyResult OnKeyProperty( FFrameNumber KeyTime, FPropertyChangedParams PropertyChangedParams )
 	{
 		FKeyPropertyResult KeyPropertyResult;
 
-		TArray<KeyDataType> NewKeysForPropertyChange;
-		TArray<KeyDataType> DefaultKeysForPropertyChange;
-		GenerateKeysFromPropertyChanged( PropertyChangedParams, NewKeysForPropertyChange, DefaultKeysForPropertyChange );
+		FGeneratedTrackKeys GeneratedKeys;
+		GenerateKeysFromPropertyChanged( PropertyChangedParams, GeneratedKeys );
 
 		UProperty* Property = PropertyChangedParams.PropertyPath.GetLeafMostProperty().Property.Get();
 		if (!Property)
@@ -271,11 +235,10 @@ private:
 		// also check for track editors which should only be used for customization.
 		if ( SupportsType( TrackClass ) && ( ForCustomizedUseOnly() == false || *CustomizedClass != nullptr) )
 		{
-			return FKeyframeTrackEditor<TrackType, SectionType, KeyDataType>::AddKeysToObjects(
+			return this->AddKeysToObjects(
 				PropertyChangedParams.ObjectsThatChanged,
 				KeyTime,
-				NewKeysForPropertyChange,
-				DefaultKeysForPropertyChange,
+				GeneratedKeys,
 				PropertyChangedParams.KeyMode,
 				TrackClass,
 				UniqueName,
diff --git a/Engine/Source/Editor/MovieSceneTools/Public/SKeyEditInterface.h b/Engine/Source/Editor/MovieSceneTools/Public/SKeyEditInterface.h
new file mode 100644
--- /dev/null
+++ b/Engine/Source/Editor/MovieSceneTools/Public/SKeyEditInterface.h
@@ -0,0 +1,66 @@
+// Copyright 1998-2018 Epic Games, Inc. All Rights Reserved.
+
+#pragma once
+
+#include "Widgets/DeclarativeSyntaxSupport.h"
+#include "Widgets/SCompoundWidget.h"
+#include "Misc/NotifyHook.h"
+#include "EditorUndoClient.h"
+#include "UObject/WeakObjectPtrTemplates.h"
+
+struct FPropertyChangedEvent;
+
+class ISequencer;
+class FStructOnScope;
+class UMovieSceneSection;
+class IPropertyTypeCustomization;
+
+struct FKeyEditData
+{
+	TSharedPtr<FStructOnScope> KeyStruct;
+	TWeakObjectPtr<UMovieSceneSection> OwningSection;
+};
+
+/**
+ * Widget that represents a details panel that refreshes on undo, and handles modification of the section on edit
+ */
+class MOVIESCENETOOLS_API SKeyEditInterface : public SCompoundWidget, public FEditorUndoClient, public FNotifyHook
+{
+public:
+	SLATE_BEGIN_ARGS(SKeyEditInterface){}
+		SLATE_ATTRIBUTE(FKeyEditData, EditData)
+	SLATE_END_ARGS()
+
+	~SKeyEditInterface();
+
+	void Construct(const FArguments& InArgs, TSharedRef<ISequencer> InSequencer);
+
+	/**
+	 * (Re)Initialize this widget's details panel
+	 */
+	void Initialize();
+
+private:
+
+	/**
+	 * Create a binding ID customization for the details panel
+	 */
+	TSharedRef<IPropertyTypeCustomization> CreateBindingIDCustomization();
+	TSharedRef<IPropertyTypeCustomization> CreateFrameNumberCustomization();
+
+	/**
+	 * Called when a property has been changed on the UI
+	 */
+	void OnFinishedChangingProperties(const FPropertyChangedEvent& ChangeEvent, TSharedPtr<FStructOnScope> KeyStruct);
+
+private:
+
+	virtual void NotifyPreChange(UProperty* PropertyAboutToChange) override;
+	virtual void PostUndo(bool bSuccess) override;
+	virtual void PostRedo(bool bSuccess) override;
+
+private:
+	TAttribute<FKeyEditData> EditDataAttribute;
+	TWeakObjectPtr<UMovieSceneSection> WeakSection;
+	TWeakPtr<ISequencer> WeakSequencer;
+};
\ No newline at end of file
diff --git a/Engine/Source/Editor/MovieSceneTools/Public/Sections/ThumbnailSection.h b/Engine/Source/Editor/MovieSceneTools/Public/Sections/ThumbnailSection.h
--- a/Engine/Source/Editor/MovieSceneTools/Public/Sections/ThumbnailSection.h
+++ b/Engine/Source/Editor/MovieSceneTools/Public/Sections/ThumbnailSection.h
@@ -62,13 +62,12 @@ public:
 public:
 
 	/** Set this thumbnail section to draw a single thumbnail at the specified time */
-	virtual void SetSingleTime(float GlobalTime) = 0;
+	virtual void SetSingleTime(double GlobalTime) = 0;
 
 public:
 
 	//~ ISequencerSection interface
 
-	virtual void GenerateSectionLayout(ISectionLayoutBuilder& LayoutBuilder) const override { }
 	virtual void BuildSectionContextMenu(FMenuBuilder& MenuBuilder, const FGuid& ObjectBinding) override;
 	virtual TSharedRef<SWidget> GenerateSectionWidget() override;
 	virtual float GetSectionGripSize() const override;
@@ -84,10 +83,10 @@ protected:
 	void RedrawThumbnails();
 
 	/** Get the range that is currently visible in the section's time space */
-	TRange<float> GetVisibleRange() const;
+	TRange<double> GetVisibleRange() const;
 
 	/** Get the total range that thumbnails are to be generated for in the section's time space */
-	TRange<float> GetTotalRange() const;
+	TRange<double> GetTotalRange() const;
 
 protected:
 
diff --git a/Engine/Source/Editor/MovieSceneTools/Public/TrackEditorThumbnail/TrackEditorThumbnail.h b/Engine/Source/Editor/MovieSceneTools/Public/TrackEditorThumbnail/TrackEditorThumbnail.h
--- a/Engine/Source/Editor/MovieSceneTools/Public/TrackEditorThumbnail/TrackEditorThumbnail.h
+++ b/Engine/Source/Editor/MovieSceneTools/Public/TrackEditorThumbnail/TrackEditorThumbnail.h
@@ -9,6 +9,7 @@
 #include "Rendering/RenderingCommon.h"
 #include "MovieSceneToolsUserSettings.h"
 #include "RHI.h"
+#include "FrameTime.h"
 
 class FLevelEditorViewportClient;
 class FSceneViewport;
@@ -30,7 +31,7 @@ class MOVIESCENETOOLS_API FTrackEditorThumbnail
 public:
 
 	/** Create and initialize a new instance. */
-	FTrackEditorThumbnail(const FOnThumbnailDraw& InOnDraw, const FIntPoint& InSize, TRange<float> InTimeRange, float InPosition);
+	FTrackEditorThumbnail(const FOnThumbnailDraw& InOnDraw, const FIntPoint& InSize, TRange<double> InTimeRange, double InPosition);
 
 	/** Virtual destructor. */
 	virtual ~FTrackEditorThumbnail();
@@ -54,10 +55,10 @@ public:
 	float GetFadeInCurve() const;
 
 	/** Get the full time-range that this thumbnail occupies */
-	const TRange<float> GetTimeRange() const { return TimeRange; }
+	const TRange<double> GetTimeRange() const { return TimeRange; }
 
 	/** Get the time at which this thumbnail should be drawn */
-	float GetEvalPosition() const { return Position; }
+	double GetEvalPosition() const { return Position; }
 
 public:
 
@@ -92,10 +93,10 @@ private:
 	FSlateTexture2DRHIRef* Texture;
 
 	/** Where in time this thumbnail is a rendering of. */
-	TRange<float> TimeRange;
+	TRange<double> TimeRange;
 
 	/** The position we should actually render (within the above time range). */
-	float Position;
+	double Position;
 
 	/** Fade curve to display while the thumbnail is redrawing. */
 	FCurveSequence FadeInCurve;
@@ -121,7 +122,7 @@ struct ICustomThumbnailClient
 /** Cache data */
 struct FThumbnailCacheData
 {
-	FThumbnailCacheData() : VisibleRange(0.f), TimeRange(0.f), AllottedSize(0,0), DesiredSize(0, 0), Quality(EThumbnailQuality::Normal) {}
+	FThumbnailCacheData() : VisibleRange(0), TimeRange(0), AllottedSize(0,0), DesiredSize(0, 0), Quality(EThumbnailQuality::Normal) {}
 
 	bool operator==(const FThumbnailCacheData& RHS) const
 	{
@@ -146,9 +147,9 @@ struct FThumbnailCacheData
 	}
 
 	/** The visible range of our thumbnails we can see on the UI */
-	TRange<float> VisibleRange;
+	TRange<double> VisibleRange;
 	/** The total range to generate thumbnails for */
-	TRange<float> TimeRange;
+	TRange<double> TimeRange;
 	/** Physical size of the thumbnail area */
 	FIntPoint AllottedSize;
 	/** Desired frame size constraint */
@@ -156,7 +157,7 @@ struct FThumbnailCacheData
 	/** Thumbnail quality */
 	EThumbnailQuality Quality;
 	/** Set when we want to render a single reference frame */
-	TOptional<float> SingleReferenceFrame;
+	TOptional<double> SingleReferenceFrame;
 };
 
 class MOVIESCENETOOLS_API FTrackEditorThumbnailCache
@@ -169,10 +170,10 @@ public:
 
 	void ForceRedraw() { bForceRedraw = true; }
 
-	void SetSingleReferenceFrame(TOptional<float> InReferenceFrame);
-	TOptional<float> GetSingleReferenceFrame() const { return CurrentCache.SingleReferenceFrame; }
+	void SetSingleReferenceFrame(TOptional<double> InReferenceFrame);
+	TOptional<double> GetSingleReferenceFrame() const { return CurrentCache.SingleReferenceFrame; }
 
-	void Update(const TRange<float>& NewRange, const TRange<float>& VisibleRange, const FIntPoint& AllottedSize, const FIntPoint& InDesiredSize, EThumbnailQuality InQuality, double InCurrentTime);
+	void Update(const TRange<double>& NewRange, const TRange<double>& VisibleRange, const FIntPoint& AllottedSize, const FIntPoint& InDesiredSize, EThumbnailQuality InQuality, double InCurrentTime);
 
 	void Revalidate(double InCurrentTime);
 
@@ -197,9 +198,9 @@ protected:
 
 	void UpdateFilledThumbnails();
 
-	void GenerateFront(const TRange<float>& Boundary);
+	void GenerateFront(const TRange<double>& Boundary);
 
-	void GenerateBack(const TRange<float>& Boundary);
+	void GenerateBack(const TRange<double>& Boundary);
 
 	void SetupViewportEngineFlags();
 
diff --git a/Engine/Source/Editor/MovieSceneTools/Public/TrackEditors/AudioTrackEditor.h b/Engine/Source/Editor/MovieSceneTools/Public/TrackEditors/AudioTrackEditor.h
--- a/Engine/Source/Editor/MovieSceneTools/Public/TrackEditors/AudioTrackEditor.h
+++ b/Engine/Source/Editor/MovieSceneTools/Public/TrackEditors/AudioTrackEditor.h
@@ -14,7 +14,6 @@
 
 struct FAssetData;
 class FAudioThumbnail;
-class FFloatCurveKeyArea;
 class FMenuBuilder;
 class FSequencerSectionPainter;
 class USoundWave;
@@ -59,14 +58,16 @@ public:
 	virtual const FSlateBrush* GetIconBrush() const override;
 	virtual bool IsResizable(UMovieSceneTrack* InTrack) const override;
 	virtual void Resize(float NewSize, UMovieSceneTrack* InTrack) override;
+	virtual bool OnAllowDrop(const FDragDropEvent& DragDropEvent, UMovieSceneTrack* Track, int32 RowIndex, const FGuid& TargetObjectGuid) override;
+	virtual FReply OnDrop(const FDragDropEvent& DragDropEvent, UMovieSceneTrack* Track, int32 RowIndex, const FGuid& TargetObjectGuid) override;
 	
 protected:
 
 	/** Delegate for AnimatablePropertyChanged in HandleAssetAdded for master sounds */
-	FKeyPropertyResult AddNewMasterSound(float KeyTime, class USoundBase* Sound);
+	FKeyPropertyResult AddNewMasterSound(FFrameNumber KeyTime, class USoundBase* Sound, int32 RowIndex);
 
 	/** Delegate for AnimatablePropertyChanged in HandleAssetAdded for attached sounds */
-	FKeyPropertyResult AddNewAttachedSound(float KeyTime, class USoundBase* Sound, TArray<TWeakObjectPtr<UObject>> ObjectsToAttachTo);
+	FKeyPropertyResult AddNewAttachedSound(FFrameNumber KeyTime, class USoundBase* Sound, TArray<TWeakObjectPtr<UObject>> ObjectsToAttachTo);
 
 private:
 
@@ -103,11 +104,10 @@ public:
 	virtual UMovieSceneSection* GetSectionObject() override;
 	virtual FText GetSectionTitle() const override;
 	virtual float GetSectionHeight() const override;
-	virtual void GenerateSectionLayout( class ISectionLayoutBuilder& LayoutBuilder) const override;
 	virtual int32 OnPaintSection(FSequencerSectionPainter& Painter) const override;
 	virtual void Tick(const FGeometry& AllottedGeometry, const FGeometry& ParentGeometry, const double InCurrentTime, const float InDeltaTime) override;
 	virtual void BeginSlipSection() override;
-	virtual void SlipSection(float SlipTime) override;
+	virtual void SlipSection(double SlipTime) override;
 	
 private:
 
@@ -122,10 +122,6 @@ private:
 	/** The waveform thumbnail render object. */
 	TSharedPtr<class FAudioThumbnail> WaveformThumbnail;
 
-	/** Sound volume/pitch key areas. */
-	mutable TSharedPtr<FFloatCurveKeyArea> SoundVolumeArea;
-	mutable TSharedPtr<FFloatCurveKeyArea> PitchMultiplierArea;
-
 	/** Stored data about the waveform to determine when it is invalidated. */
 	TRange<float> StoredDrawRange;
 	float StoredStartOffset;
@@ -143,5 +139,5 @@ private:
 	float InitialStartOffsetDuringResize;
 	
 	/** Cached start time valid only during resize */
-	float InitialStartTimeDuringResize;
+	double InitialStartTimeDuringResize;
 };
diff --git a/Engine/Source/Editor/MovieSceneTools/Public/TrackEditors/PropertyTrackEditors/BoolPropertyTrackEditor.h b/Engine/Source/Editor/MovieSceneTools/Public/TrackEditors/PropertyTrackEditors/BoolPropertyTrackEditor.h
--- a/Engine/Source/Editor/MovieSceneTools/Public/TrackEditors/PropertyTrackEditors/BoolPropertyTrackEditor.h
+++ b/Engine/Source/Editor/MovieSceneTools/Public/TrackEditors/PropertyTrackEditors/BoolPropertyTrackEditor.h
@@ -15,7 +15,7 @@
  * A property track editor for Booleans.
  */
 class MOVIESCENETOOLS_API FBoolPropertyTrackEditor
-	: public FPropertyTrackEditor<UMovieSceneBoolTrack, UMovieSceneBoolSection, bool>
+	: public FPropertyTrackEditor<UMovieSceneBoolTrack>
 {
 public:
 
@@ -44,13 +44,9 @@ public:
 	 */
 	static TSharedRef<ISequencerTrackEditor> CreateTrackEditor( TSharedRef<ISequencer> OwningSequencer );
 
-	//~ ISequencerTrackEditor interface
-
-	virtual TSharedRef<ISequencerSection> MakeSectionInterface(UMovieSceneSection& SectionObject, UMovieSceneTrack& Track, FGuid ObjectBinding) override;
-
 protected:
 
 	//~ FPropertyTrackEditor interface
-
-	virtual void GenerateKeysFromPropertyChanged( const FPropertyChangedParams& PropertyChangedParams, TArray<bool>& NewGeneratedKeys, TArray<bool>& DefaultGeneratedKeys ) override;
+	virtual TSharedRef<ISequencerSection> MakeSectionInterface(UMovieSceneSection& SectionObject, UMovieSceneTrack& Track, FGuid ObjectBinding) override;
+	virtual void GenerateKeysFromPropertyChanged( const FPropertyChangedParams& PropertyChangedParams, FGeneratedTrackKeys& OutGeneratedKeys ) override;
 };
diff --git a/Engine/Source/Editor/MovieSceneTools/Public/TrackEditors/SpawnTrackEditor.h b/Engine/Source/Editor/MovieSceneTools/Public/TrackEditors/SpawnTrackEditor.h
--- a/Engine/Source/Editor/MovieSceneTools/Public/TrackEditors/SpawnTrackEditor.h
+++ b/Engine/Source/Editor/MovieSceneTools/Public/TrackEditors/SpawnTrackEditor.h
@@ -48,7 +48,6 @@ public:
 	virtual bool HandleAssetAdded(UObject* Asset, const FGuid& TargetObjectGuid) override { return false; }
 	virtual bool SupportsType(TSubclassOf<UMovieSceneTrack> Type) const override;
     virtual bool SupportsSequence(UMovieSceneSequence* InSequence) const override;
-	virtual TSharedRef<ISequencerSection> MakeSectionInterface(UMovieSceneSection& SectionObject, UMovieSceneTrack& Track, FGuid ObjectBinding) override;
 
 private:
 
diff --git a/Engine/Source/Editor/MovieSceneTools/Public/TrackEditors/SubTrackEditor.h b/Engine/Source/Editor/MovieSceneTools/Public/TrackEditors/SubTrackEditor.h
--- a/Engine/Source/Editor/MovieSceneTools/Public/TrackEditors/SubTrackEditor.h
+++ b/Engine/Source/Editor/MovieSceneTools/Public/TrackEditors/SubTrackEditor.h
@@ -53,9 +53,12 @@ public:
 	virtual bool SupportsSequence(UMovieSceneSequence* InSequence) const override;
 	virtual bool SupportsType(TSubclassOf<UMovieSceneTrack> Type) const override;
 	virtual const FSlateBrush* GetIconBrush() const override;
-	virtual bool OnAllowDrop(const FDragDropEvent& DragDropEvent, UMovieSceneTrack* Track) override;
-	virtual FReply OnDrop(const FDragDropEvent& DragDropEvent, UMovieSceneTrack* Track) override;
+	virtual bool OnAllowDrop(const FDragDropEvent& DragDropEvent, UMovieSceneTrack* Track, int32 RowIndex, const FGuid& TargetObjectGuid) override;
+	virtual FReply OnDrop(const FDragDropEvent& DragDropEvent, UMovieSceneTrack* Track, int32 RowIndex, const FGuid& TargetObjectGuid) override;
 	
+	/** Switch the sub section's take */
+	void SwitchTake(UMovieSceneSubSection* Section, uint32 TakeNumber);
+
 protected:
 
 	/**
@@ -81,10 +84,10 @@ private:
 	void HandleAddSubSequenceComboButtonMenuEntryExecute(const FAssetData& AssetData, UMovieSceneTrack* InTrack);
 
 	/** Delegate for AnimatablePropertyChanged in AddKey */
-	FKeyPropertyResult AddKeyInternal(float KeyTime, UMovieSceneSequence* InMovieSceneSequence, UMovieSceneTrack* InTrack);
+	FKeyPropertyResult AddKeyInternal(FFrameNumber KeyTime, UMovieSceneSequence* InMovieSceneSequence, UMovieSceneTrack* InTrack, int32 RowIndex);
 
 	/** Callback for AnimatablePropertyChanged in HandleAssetAdded. */
-	FKeyPropertyResult HandleSequenceAdded(float KeyTime, UMovieSceneSequence* Sequence);
+	FKeyPropertyResult HandleSequenceAdded(FFrameNumber KeyTime, UMovieSceneSequence* Sequence, int32 RowIndex);
 
 	/** Check if we can record a new sequence (deny it if one is already primed) */
 	bool CanRecordNewSequence() const;
@@ -93,5 +96,5 @@ private:
 	void HandleRecordNewSequence(AActor* InActorToRecord, UMovieSceneTrack* InTrack);
 
 	/** Actually handles the adding of the section */
-	FKeyPropertyResult HandleRecordNewSequenceInternal(float KeyTime, AActor* InActorToRecord, UMovieSceneTrack* InTrack);
+	FKeyPropertyResult HandleRecordNewSequenceInternal(FFrameNumber KeyTime, AActor* InActorToRecord, UMovieSceneTrack* InTrack);
 };
diff --git a/Engine/Source/Editor/PropertyEditor/Private/DetailItemNode.cpp b/Engine/Source/Editor/PropertyEditor/Private/DetailItemNode.cpp
--- a/Engine/Source/Editor/PropertyEditor/Private/DetailItemNode.cpp
+++ b/Engine/Source/Editor/PropertyEditor/Private/DetailItemNode.cpp
@@ -5,7 +5,8 @@
 #include "DetailGroup.h"
 #include "DetailPropertyRow.h"
 #include "SDetailSingleItemRow.h"
-
+#include "IDetailKeyframeHandler.h"
+#include "ObjectPropertyNode.h"
 
 
 FDetailItemNode::FDetailItemNode(const FDetailLayoutCustomization& InCustomization, TSharedRef<FDetailCategoryImpl> InParentCategory, TAttribute<bool> InIsParentEnabled, TSharedPtr<IDetailGroup> InParentGroup)
@@ -392,7 +393,7 @@ ENodeVisibility FDetailItemNode::GetVisibility() const
 	return Visibility;
 }
 
-static bool PassesAllFilters( const FDetailLayoutCustomization& InCustomization, const FDetailFilter& InFilter, const FString& InCategoryName )
+static bool PassesAllFilters( FDetailItemNode* ItemNode, const FDetailLayoutCustomization& InCustomization, const FDetailFilter& InFilter, const FString& InCategoryName )
 {	
 	struct Local
 	{
@@ -413,22 +414,60 @@ static bool PassesAllFilters( const FDetailLayoutCustomization& InCustomization,
 			}
 			return false;
 		}
+		static bool ItemIsKeyable( FDetailItemNode *InItemNode, UClass *ObjectClass, TSharedPtr<FPropertyNode> PropertyNode)
+		{
+			IDetailsViewPrivate* DetailsView = InItemNode->GetDetailsView();
+			TSharedPtr<IDetailKeyframeHandler> KeyframeHandler = DetailsView->GetKeyframeHandler();
+			if (KeyframeHandler.IsValid())
+			{
+				TSharedPtr<IPropertyHandle> PropertyHandle = PropertyEditorHelpers::GetPropertyHandle(PropertyNode.ToSharedRef(), nullptr, nullptr);
+
+				return
+					KeyframeHandler->IsPropertyKeyingEnabled() &&
+					KeyframeHandler->IsPropertyKeyable(ObjectClass, *PropertyHandle);
+			}
+			return false;
+		}
+
+		static bool ItemIsAnimated(FDetailItemNode *InItemNode, TSharedPtr<FPropertyNode> PropertyNode)
+		{
+			IDetailsViewPrivate* DetailsView = InItemNode->GetDetailsView();
+			TSharedPtr<IDetailKeyframeHandler> KeyframeHandler = DetailsView->GetKeyframeHandler();
+			if (KeyframeHandler.IsValid())
+			{
+				TSharedPtr<IPropertyHandle> PropertyHandle = PropertyEditorHelpers::GetPropertyHandle(PropertyNode.ToSharedRef(), nullptr, nullptr);
+				FObjectPropertyNode *ParentPropertyNode = PropertyNode->FindObjectItemParent();
+				// Get an iterator for the enclosing objects.
+				for (int32 ObjIndex = 0; ObjIndex < ParentPropertyNode->GetNumObjects(); ++ObjIndex)
+				{
+					UObject* ParentObject = ParentPropertyNode->GetUObject(ObjIndex);
+
+					if (KeyframeHandler->IsPropertyAnimated(*PropertyHandle, ParentObject))
+					{
+						return true;
+					}
+				}
+			}
+			return false;
+		}
 	};
 
 	bool bPassesAllFilters = true;
 
-	if( InFilter.FilterStrings.Num() > 0 || InFilter.bShowOnlyModifiedProperties == true || InFilter.bShowOnlyDiffering == true )
+	if( InFilter.FilterStrings.Num() > 0 || InFilter.bShowOnlyModifiedProperties == true || InFilter.bShowOnlyDiffering == true 
+		|| InFilter.bShowKeyable == true || InFilter.bShowAnimated == true)
 	{
 		const bool bSearchFilterIsEmpty = InFilter.FilterStrings.Num() == 0;
 
 		TSharedPtr<FPropertyNode> PropertyNodePin = InCustomization.GetPropertyNode();
-
+		
 		const bool bPassesCategoryFilter = !bSearchFilterIsEmpty && InFilter.bShowAllChildrenIfCategoryMatches ? Local::StringPassesFilter(InFilter, InCategoryName) : false;
 
 		bPassesAllFilters = false;
 		if( PropertyNodePin.IsValid() && !PropertyNodePin->AsCategoryNode() )
 		{
-			
+			UClass* ObjectClass = PropertyNodePin->FindObjectItemParent()->GetObjectBaseClass();
+
 			const bool bIsNotBeingFiltered = PropertyNodePin->HasNodeFlags(EPropertyNodeFlags::IsBeingFiltered) == 0;
 			const bool bIsSeenDueToFiltering = PropertyNodePin->HasNodeFlags(EPropertyNodeFlags::IsSeenDueToFiltering) != 0;
 			const bool bIsParentSeenDueToFiltering = PropertyNodePin->HasNodeFlags(EPropertyNodeFlags::IsParentSeenDueToFiltering) != 0;
@@ -436,16 +475,20 @@ static bool PassesAllFilters( const FDetailLayoutCustomization& InCustomization,
 			const bool bPassesSearchFilter = bSearchFilterIsEmpty || ( bIsNotBeingFiltered || bIsSeenDueToFiltering || bIsParentSeenDueToFiltering );
 			const bool bPassesModifiedFilter = bPassesSearchFilter && ( InFilter.bShowOnlyModifiedProperties == false || PropertyNodePin->GetDiffersFromDefault() == true );
 			const bool bPassesDifferingFilter = InFilter.bShowOnlyDiffering ? InFilter.WhitelistedProperties.Find(*FPropertyNode::CreatePropertyPath(PropertyNodePin.ToSharedRef())) != nullptr : true;
+			const bool bPassesKeyableFilter = (InFilter.bShowKeyable == false || Local::ItemIsKeyable(ItemNode, ObjectClass, PropertyNodePin));
+			const bool bPassesAnimatedFilter = (InFilter.bShowAnimated == false || Local::ItemIsAnimated(ItemNode, PropertyNodePin));
 
 			// The property node is visible (note categories are never visible unless they have a child that is visible )
-			bPassesAllFilters = (bPassesSearchFilter && bPassesModifiedFilter && bPassesDifferingFilter) || bPassesCategoryFilter;
+			bPassesAllFilters = (bPassesSearchFilter && bPassesModifiedFilter && bPassesDifferingFilter && bPassesKeyableFilter && bPassesAnimatedFilter) || bPassesCategoryFilter;
 		}
 		else if (InCustomization.HasCustomWidget())
 		{
 			const bool bPassesTextFilter = Local::StringPassesFilter(InFilter, InCustomization.WidgetDecl->FilterTextString.ToString());
 			const bool bPassesModifiedFilter = (InFilter.bShowOnlyModifiedProperties == false || InCustomization.WidgetDecl->DiffersFromDefaultAttr.Get() == true);
-
-			bPassesAllFilters = (bPassesTextFilter && bPassesModifiedFilter) || bPassesCategoryFilter;
+			//@todo we need to support custom widgets for keyable,animated, in particularly for transforms(ComponentTransformDetails).
+			const bool bPassesKeyableFilter = (InFilter.bShowKeyable == false);
+			const bool bPassesAnimatedFilter = (InFilter.bShowAnimated == false);
+			bPassesAllFilters = (bPassesTextFilter && bPassesModifiedFilter && bPassesKeyableFilter && bPassesAnimatedFilter) || bPassesCategoryFilter;
 		}
 	}
 
@@ -539,7 +582,7 @@ TSharedPtr<IDetailPropertyRow> FDetailItemNode::GetRow() const
 
 void FDetailItemNode::FilterNode(const FDetailFilter& InFilter)
 {
-	bShouldBeVisibleDueToFiltering = PassesAllFilters( Customization, InFilter, ParentCategory.Pin()->GetDisplayName().ToString() );
+	bShouldBeVisibleDueToFiltering = PassesAllFilters( this, Customization, InFilter, ParentCategory.Pin()->GetDisplayName().ToString() );
 
 	bShouldBeVisibleDueToChildFiltering = false;
 
diff --git a/Engine/Source/Editor/PropertyEditor/Private/SDetailsView.cpp b/Engine/Source/Editor/PropertyEditor/Private/SDetailsView.cpp
--- a/Engine/Source/Editor/PropertyEditor/Private/SDetailsView.cpp
+++ b/Engine/Source/Editor/PropertyEditor/Private/SDetailsView.cpp
@@ -87,7 +87,36 @@ void SDetailsView::Construct(const FArguments& InArgs)
 				EUserInterfaceActionType::ToggleButton
 			);
 		}
-
+		if (DetailsViewArgs.bShowKeyablePropertiesOption)
+		{
+			DetailViewOptions.AddMenuEntry(
+				LOCTEXT("ShowOnlyKeyable", "Show Only Keyable Properties"),
+				LOCTEXT("ShowOnlyKeyable_ToolTip", "Displays only properties which are keyable"),
+				FSlateIcon(),
+				FUIAction(
+					FExecuteAction::CreateSP(this, &SDetailsView::OnShowKeyableClicked),
+					FCanExecuteAction(),
+					FIsActionChecked::CreateSP(this, &SDetailsView::IsShowKeyableChecked)
+				),
+				NAME_None,
+				EUserInterfaceActionType::ToggleButton
+			);
+		}
+		if (DetailsViewArgs.bShowAnimatedPropertiesOption)
+		{
+			DetailViewOptions.AddMenuEntry(
+				LOCTEXT("ShowAnimated", "Show Only Animated Properties"),
+				LOCTEXT("ShowAnimated_ToolTip", "Displays only properties which are animated (have tracks)"),
+				FSlateIcon(),
+				FUIAction(
+					FExecuteAction::CreateSP(this, &SDetailsView::OnShowAnimatedClicked),
+					FCanExecuteAction(),
+					FIsActionChecked::CreateSP(this, &SDetailsView::IsShowAnimatedChecked)
+				),
+				NAME_None,
+				EUserInterfaceActionType::ToggleButton
+			);
+		}
 		DetailViewOptions.AddMenuEntry(
 			LOCTEXT("ShowAllAdvanced", "Show All Advanced Details"),
 			LOCTEXT("ShowAllAdvanced_ToolTip", "Shows all advanced detail sections in each category"),
diff --git a/Engine/Source/Editor/PropertyEditor/Private/SDetailsViewBase.cpp b/Engine/Source/Editor/PropertyEditor/Private/SDetailsViewBase.cpp
--- a/Engine/Source/Editor/PropertyEditor/Private/SDetailsViewBase.cpp
+++ b/Engine/Source/Editor/PropertyEditor/Private/SDetailsViewBase.cpp
@@ -548,6 +548,19 @@ void SDetailsViewBase::OnShowAllChildrenIfCategoryMatchesClicked()
 	UpdateFilteredDetails();
 }
 
+void SDetailsViewBase::OnShowKeyableClicked()
+{
+	CurrentFilter.bShowKeyable = !CurrentFilter.bShowKeyable;
+
+	UpdateFilteredDetails();
+}
+
+void SDetailsViewBase::OnShowAnimatedClicked()
+{
+	CurrentFilter.bShowAnimated = !CurrentFilter.bShowAnimated;
+
+	UpdateFilteredDetails();
+}
 /** Called when the filter text changes.  This filters specific property nodes out of view */
 void SDetailsViewBase::OnFilterTextChanged(const FText& InFilterText)
 {
@@ -612,7 +625,7 @@ void SDetailsViewBase::FilterView(const FString& InFilterText)
 EVisibility SDetailsViewBase::GetFilterBoxVisibility() const
 {
 	// Visible if we allow search and we have anything to search otherwise collapsed so it doesn't take up room
-	return (DetailsViewArgs.bAllowSearch && IsConnected() && RootTreeNodes.Num() > 0) || HasActiveSearch() || CurrentFilter.bShowOnlyModifiedProperties || CurrentFilter.bShowOnlyDiffering ? EVisibility::Visible : EVisibility::Collapsed;
+	return (DetailsViewArgs.bAllowSearch && IsConnected() && RootTreeNodes.Num() > 0) || HasActiveSearch() || CurrentFilter.bShowOnlyModifiedProperties || CurrentFilter.bShowAnimated || CurrentFilter.bShowKeyable || CurrentFilter.bShowOnlyDiffering ? EVisibility::Visible : EVisibility::Collapsed;
 }
 
 bool SDetailsViewBase::SupportsKeyboardFocus() const
diff --git a/Engine/Source/Editor/PropertyEditor/Private/SDetailsViewBase.h b/Engine/Source/Editor/PropertyEditor/Private/SDetailsViewBase.h
--- a/Engine/Source/Editor/PropertyEditor/Private/SDetailsViewBase.h
+++ b/Engine/Source/Editor/PropertyEditor/Private/SDetailsViewBase.h
@@ -44,9 +44,11 @@ struct FDetailFilter
 		, bShowAllAdvanced(false)
 		, bShowOnlyDiffering(false)
 		, bShowAllChildrenIfCategoryMatches(true)
+		, bShowKeyable(false)
+		, bShowAnimated(false)
 	{}
 
-	bool IsEmptyFilter() const { return FilterStrings.Num() == 0 && bShowOnlyModifiedProperties == false && bShowAllAdvanced == false && bShowOnlyDiffering == false && bShowAllChildrenIfCategoryMatches == false; }
+	bool IsEmptyFilter() const { return FilterStrings.Num() == 0 && bShowOnlyModifiedProperties == false && bShowAllAdvanced == false && bShowOnlyDiffering == false && bShowAllChildrenIfCategoryMatches == false && bShowKeyable == false && bShowAnimated == false; }
 
 	/** Any user search terms that items must match */
 	TArray<FString> FilterStrings;
@@ -58,6 +60,10 @@ struct FDetailFilter
 	bool bShowOnlyDiffering;
 	/** If we should show all the children if their category name matches the search */
 	bool bShowAllChildrenIfCategoryMatches;
+	/** If we should only show keyable properties */
+	bool bShowKeyable;
+	/** If we should only show animated properties */
+	bool bShowAnimated;
 	TSet<FPropertyPath> WhitelistedProperties;
 };
 
@@ -270,6 +276,12 @@ protected:
 
 	/** @return true if show all advanced is checked */
 	bool IsShowAllChildrenIfCategoryMatchesChecked() const { return CurrentFilter.bShowAllChildrenIfCategoryMatches; }
+	
+	/** @return true if show keyable is checked */
+	bool IsShowKeyableChecked() const { return CurrentFilter.bShowKeyable; }
+	
+	/** @return true if show animated is checked */
+	bool IsShowAnimatedChecked() const { return CurrentFilter.bShowAnimated; }
 
 	/** Called when show only modified is clicked */
 	void OnShowOnlyModifiedClicked();
@@ -283,6 +295,12 @@ protected:
 	/** Called when show all children if category matches is clicked */
 	void OnShowAllChildrenIfCategoryMatchesClicked();
 
+	/** Called when show keyable is clicked */
+	void OnShowKeyableClicked();
+	
+	/** Calledw when show animated is clicked */
+	void OnShowAnimatedClicked();
+
 	/**
 	* Updates the details with the passed in filter
 	*/
diff --git a/Engine/Source/Editor/PropertyEditor/Private/SStructureDetailsView.cpp b/Engine/Source/Editor/PropertyEditor/Private/SStructureDetailsView.cpp
--- a/Engine/Source/Editor/PropertyEditor/Private/SStructureDetailsView.cpp
+++ b/Engine/Source/Editor/PropertyEditor/Private/SStructureDetailsView.cpp
@@ -72,7 +72,36 @@ void SStructureDetailsView::Construct(const FArguments& InArgs)
 				EUserInterfaceActionType::ToggleButton 
 			);
 		}
-
+		if (DetailsViewArgs.bShowKeyablePropertiesOption)
+		{
+			DetailViewOptions.AddMenuEntry(
+				LOCTEXT("ShowOnlyKeyable", "Show Only Keyable Properties"),
+				LOCTEXT("ShowOnlyKeyable_ToolTip", "Displays only properties which are keyable"),
+				FSlateIcon(),
+				FUIAction(
+					FExecuteAction::CreateSP(this, &SStructureDetailsView::OnShowKeyableClicked),
+					FCanExecuteAction(),
+					FIsActionChecked::CreateSP(this, &SStructureDetailsView::IsShowKeyableChecked)
+				),
+				NAME_None,
+				EUserInterfaceActionType::ToggleButton
+			);
+		}
+		if (DetailsViewArgs.bShowAnimatedPropertiesOption)
+		{
+			DetailViewOptions.AddMenuEntry(
+				LOCTEXT("ShowAnimated", "Show Only Animated Properties"),
+				LOCTEXT("ShowAnimated_ToolTip", "Displays only properties which are animated (have tracks)"),
+				FSlateIcon(),
+				FUIAction(
+					FExecuteAction::CreateSP(this, &SStructureDetailsView::OnShowAnimatedClicked),
+					FCanExecuteAction(),
+					FIsActionChecked::CreateSP(this, &SStructureDetailsView::IsShowAnimatedChecked)
+				),
+				NAME_None,
+				EUserInterfaceActionType::ToggleButton
+			);
+		}
 		FUIAction ShowAllAdvancedAction( 
 			FExecuteAction::CreateSP(this, &SStructureDetailsView::OnShowAllAdvancedClicked),
 			FCanExecuteAction(),
diff --git a/Engine/Source/Editor/PropertyEditor/Public/IDetailKeyframeHandler.h b/Engine/Source/Editor/PropertyEditor/Public/IDetailKeyframeHandler.h
--- a/Engine/Source/Editor/PropertyEditor/Public/IDetailKeyframeHandler.h
+++ b/Engine/Source/Editor/PropertyEditor/Public/IDetailKeyframeHandler.h
@@ -17,4 +17,6 @@ public:
 
 	virtual void OnKeyPropertyClicked(const IPropertyHandle& KeyedPropertyHandle) = 0;
 
+	virtual bool IsPropertyAnimated(const class IPropertyHandle& PropertyHandle, UObject *ParentObject) const = 0;
+
 };
diff --git a/Engine/Source/Editor/PropertyEditor/Public/IDetailsView.h b/Engine/Source/Editor/PropertyEditor/Public/IDetailsView.h
--- a/Engine/Source/Editor/PropertyEditor/Public/IDetailsView.h
+++ b/Engine/Source/Editor/PropertyEditor/Public/IDetailsView.h
@@ -90,6 +90,11 @@ struct FDetailsViewArgs
 	uint32 bShowScrollBar : 1;
 	/** If true, all properties will be visible, not just those with CPF_Edit */
 	uint32 bForceHiddenPropertyVisibility : 1;
+	/** True if you want to show the 'Show Only Keyable Properties'. Only valid in conjunction with bShowOptions */
+	uint32 bShowKeyablePropertiesOption : 1;
+	/** True if you want to show the 'Show Only Animated Properties'. Only valid in conjunction with bShowOptions */
+	uint32 bShowAnimatedPropertiesOption: 1;
+
 public:
 	/** Default constructor */
 	FDetailsViewArgs( const bool InUpdateFromSelection = false
@@ -120,6 +125,8 @@ public:
 		, bAllowMultipleTopLevelObjects(false)
 		, bShowScrollBar(true)
 		, bForceHiddenPropertyVisibility(false)
+		, bShowKeyablePropertiesOption(true)
+		, bShowAnimatedPropertiesOption(true)
 	{
 	}
 };
diff --git a/Engine/Source/Editor/SequenceRecorder/Private/ActorRecording.cpp b/Engine/Source/Editor/SequenceRecorder/Private/ActorRecording.cpp
--- a/Engine/Source/Editor/SequenceRecorder/Private/ActorRecording.cpp
+++ b/Engine/Source/Editor/SequenceRecorder/Private/ActorRecording.cpp
@@ -22,6 +22,7 @@ static const FName MovieSceneSectionRecorderFactoryName("MovieSceneSectionRecord
 
 UActorRecording::UActorRecording(const FObjectInitializer& ObjectInitializer)
 	: Super(ObjectInitializer)
+	, ActorSettings(this)
 {
 	bActive = true;
 	bWasSpawnedPostRecord = false;
@@ -328,7 +329,10 @@ void UActorRecording::StartRecordingActorProperties(ULevelSequence* CurrentSeque
 			// force set recording to record translations as we need this with no animation
 			UMovieScene3DTransformSectionRecorderSettings* TransformSettings = ActorSettings.GetSettingsObject<UMovieScene3DTransformSectionRecorderSettings>();
 			check(TransformSettings);
-			TransformSettings->bRecordTransforms = true;
+			if (!TransformSettings->bRecordTransforms)
+			{
+				UE_LOG(LogAnimation, Warning, TEXT("Transform recording is not enabled for '%s'. Resulting animation may not match gameplay."), *Actor->GetActorLabel());
+			}
 
 			// grab components so we can track attachments
 			// don't include non-CDO here as they wont be part of our initial BP (duplicated above)
diff --git a/Engine/Source/Editor/SequenceRecorder/Private/ActorRecording.h b/Engine/Source/Editor/SequenceRecorder/Private/ActorRecording.h
--- a/Engine/Source/Editor/SequenceRecorder/Private/ActorRecording.h
+++ b/Engine/Source/Editor/SequenceRecorder/Private/ActorRecording.h
@@ -17,7 +17,7 @@
 class ULevelSequence;
 class UMovieScene;
 
-UCLASS(MinimalAPI, Transient)
+UCLASS(MinimalAPI)
 class UActorRecording : public UObject
 {
 	GENERATED_UCLASS_BODY()
@@ -85,7 +85,7 @@ public:
 	UPROPERTY(EditAnywhere, Category = "Actor Recording")
 	FActorRecordingSettings ActorSettings;
 
-	/** Whether this recording is active and to be recorded when the 'Record' button is pressed. */
+	/** Whether this actor is active and to be recorded when the 'Record' button is pressed. */
 	UPROPERTY(EditAnywhere, Category = "Actor Recording")
 	bool bActive;
 
diff --git a/Engine/Source/Editor/SequenceRecorder/Private/ActorRecordingSettings.cpp b/Engine/Source/Editor/SequenceRecorder/Private/ActorRecordingSettings.cpp
--- a/Engine/Source/Editor/SequenceRecorder/Private/ActorRecordingSettings.cpp
+++ b/Engine/Source/Editor/SequenceRecorder/Private/ActorRecordingSettings.cpp
@@ -5,14 +5,25 @@
 #include "Features/IModularFeatures.h"
 
 FActorRecordingSettings::FActorRecordingSettings()
+{
+	CreateSettingsObjectsFromFactory();
+}
+
+FActorRecordingSettings::FActorRecordingSettings(UObject* InOuter)
+{
+	Outer = InOuter;
+	CreateSettingsObjectsFromFactory();
+}
+
+void FActorRecordingSettings::CreateSettingsObjectsFromFactory()
 {
 	TArray<IMovieSceneSectionRecorderFactory*> ModularFeatures = IModularFeatures::Get().GetModularFeatureImplementations<IMovieSceneSectionRecorderFactory>("MovieSceneSectionRecorderFactory");
 	for (IMovieSceneSectionRecorderFactory* Factory : ModularFeatures)
 	{
-		UObject* SettingsObject = Factory->CreateSettingsObject();
+		UObject* SettingsObject = Factory->CreateSettingsObject(Outer.Get() ? Outer.Get() : (UObject*)GetTransientPackage());
 		if (SettingsObject)
 		{
 			Settings.Add(SettingsObject);
 		}
 	}
-}
+}
\ No newline at end of file
diff --git a/Engine/Source/Editor/SequenceRecorder/Private/AnimationRecorder.cpp b/Engine/Source/Editor/SequenceRecorder/Private/AnimationRecorder.cpp
--- a/Engine/Source/Editor/SequenceRecorder/Private/AnimationRecorder.cpp
+++ b/Engine/Source/Editor/SequenceRecorder/Private/AnimationRecorder.cpp
@@ -510,7 +510,11 @@ void FAnimationRecorder::UpdateRecord(USkeletalMeshComponent* Component, float D
 				BlendedCurve = AnimCurves;
 			}
 
-			Record(Component, BlendedComponentToWorld, BlendedSpaceBases, BlendedCurve, FramesRecorded + 1);
+			if (!Record(Component, BlendedComponentToWorld, BlendedSpaceBases, BlendedCurve, FramesRecorded + 1))
+			{
+				StopRecord(true);
+				return;
+			}
 			++FramesRecorded;
 		}
 	}
@@ -528,7 +532,7 @@ void FAnimationRecorder::UpdateRecord(USkeletalMeshComponent* Component, float D
 	}
 }
 
-void FAnimationRecorder::Record(USkeletalMeshComponent* Component, FTransform const& ComponentToWorld, const TArray<FTransform>& SpacesBases, const FBlendedHeapCurve& AnimationCurves, int32 FrameToAdd)
+bool FAnimationRecorder::Record(USkeletalMeshComponent* Component, FTransform const& ComponentToWorld, const TArray<FTransform>& SpacesBases, const FBlendedHeapCurve& AnimationCurves, int32 FrameToAdd)
 {
 	if (ensure(AnimationObject))
 	{
@@ -611,7 +615,11 @@ void FAnimationRecorder::Record(USkeletalMeshComponent* Component, FTransform co
 				RawTrack.ScaleKeys.Add(LocalTransform.GetScale3D());
 
 				// verification
-				check (FrameToAdd == RawTrack.PosKeys.Num()-1);
+				if (FrameToAdd != RawTrack.PosKeys.Num()-1)
+				{
+					UE_LOG(LogAnimation, Warning, TEXT("Mismatch in animation frames. Trying to record frame: %d, but only: %d frame(s) exist. Changing skeleton while recording is not supported."), FrameToAdd, RawTrack.PosKeys.Num());
+					return false;
+				}
 			}
 		}
 
@@ -631,6 +639,8 @@ void FAnimationRecorder::Record(USkeletalMeshComponent* Component, FTransform co
 
 		LastFrame = FrameToAdd;
 	}
+
+	return true;
 }
 
 void FAnimationRecorder::RecordNotifies(USkeletalMeshComponent* Component, const TArray<FAnimNotifyEventReference>& AnimNotifies, float DeltaTime, float RecordTime)
diff --git a/Engine/Source/Editor/SequenceRecorder/Private/AnimationRecorder.h b/Engine/Source/Editor/SequenceRecorder/Private/AnimationRecorder.h
--- a/Engine/Source/Editor/SequenceRecorder/Private/AnimationRecorder.h
+++ b/Engine/Source/Editor/SequenceRecorder/Private/AnimationRecorder.h
@@ -84,7 +84,7 @@ public:
 	ERichCurveTangentMode TangentMode;
 
 private:
-	void Record(USkeletalMeshComponent* Component, FTransform const& ComponentToWorld, const TArray<FTransform>& SpacesBases, const FBlendedHeapCurve& AnimationCurves, int32 FrameToAdd);
+	bool Record(USkeletalMeshComponent* Component, FTransform const& ComponentToWorld, const TArray<FTransform>& SpacesBases, const FBlendedHeapCurve& AnimationCurves, int32 FrameToAdd);
 
 	void RecordNotifies(USkeletalMeshComponent* Component, const TArray<FAnimNotifyEventReference>& AnimNotifies, float DeltaTime, float RecordTime);
 
diff --git a/Engine/Source/Editor/SequenceRecorder/Private/SSequenceRecorder.cpp b/Engine/Source/Editor/SequenceRecorder/Private/SSequenceRecorder.cpp
--- a/Engine/Source/Editor/SequenceRecorder/Private/SSequenceRecorder.cpp
+++ b/Engine/Source/Editor/SequenceRecorder/Private/SSequenceRecorder.cpp
@@ -23,6 +23,18 @@
 #include "SequenceRecorderSettings.h"
 #include "SequenceRecorder.h"
 #include "SDropTarget.h"
+#include "EditorFontGlyphs.h"
+#include "Widgets/Input/SComboButton.h"
+#include "Widgets/Input/SEditableTextBox.h"
+#include "Widgets/Layout/SUniformGridPanel.h"
+#include "Framework/Application/SlateApplication.h"
+#include "SequenceRecorderActorGroup.h"
+#include "ISinglePropertyView.h"
+#include "ActorGroupDetailsCustomization.h"
+#include "EditorSupportDelegates.h"
+#include "Editor.h"
+#include "LevelEditor.h"
+
 
 #define LOCTEXT_NAMESPACE "SequenceRecorder"
 
@@ -175,7 +187,13 @@ void SSequenceRecorder::Construct(const FArguments& Args)
 
 	ActorRecordingDetailsView = PropertyEditorModule.CreateDetailView( DetailsViewArgs );
 	SequenceRecordingDetailsView = PropertyEditorModule.CreateDetailView( DetailsViewArgs );
+	RecordingGroupDetailsView = PropertyEditorModule.CreateDetailView( DetailsViewArgs );
+
+	TWeakPtr<SSequenceRecorder> WeakPtr = SharedThis(this);
+	RecordingGroupDetailsView->RegisterInstancedCustomPropertyLayout(USequenceRecorderActorGroup::StaticClass(), FOnGetDetailCustomizationInstance::CreateStatic(&FActorGroupDetailsCustomization::MakeInstance, WeakPtr));
+
 	SequenceRecordingDetailsView->SetObject(GetMutableDefault<USequenceRecorderSettings>());
+	RecordingGroupDetailsView->SetObject(GetMutableDefault<USequenceRecorderActorGroup>());
 
 	FToolBarBuilder ToolBarBuilder(CommandList, FMultiBoxCustomization::None);
 
@@ -282,6 +300,10 @@ void SSequenceRecorder::Construct(const FArguments& Args)
 				SNew(SVerticalBox)
 				.IsEnabled_Lambda([]() { return !FSequenceRecorder::Get().IsRecording(); })
 				+SVerticalBox::Slot()
+				[
+					RecordingGroupDetailsView.ToSharedRef()
+				]
+				+SVerticalBox::Slot()
 				.AutoHeight()
 				[
 					SequenceRecordingDetailsView.ToSharedRef()
@@ -300,6 +322,38 @@ void SSequenceRecorder::Construct(const FArguments& Args)
 	{
 		ActiveTimerHandle = RegisterActiveTimer(0.1f, FWidgetActiveTimerDelegate::CreateSP(this, &SSequenceRecorder::HandleRefreshItems));
 	}
+
+#if WITH_EDITOR
+	FEditorSupportDelegates::PrepareToCleanseEditorObject.AddRaw(this, &SSequenceRecorder::HandleMapUnload);
+#endif
+}
+
+SSequenceRecorder::~SSequenceRecorder()
+{
+#if WITH_EDITOR
+	FEditorSupportDelegates::PrepareToCleanseEditorObject.RemoveAll(this);
+#endif
+}
+
+void SSequenceRecorder::HandleMapUnload(UObject* Object)
+{
+	UWorld* EditorWorld = GEditor->GetEditorWorldContext().World();
+	if (EditorWorld == Object)
+	{
+		// When a map object is about to be GC'd we want to make sure the UI releases all references
+		// to anything that is owned by the scene.
+		ActorRecordingDetailsView->SetObject(nullptr);
+
+		// Force the list view to rebuild after clearing it's data source. This clears the list view's
+		// widget children. The ListView should only contain TWeakObjectPtrs but it's holding a hard
+		// reference anyways that's causes a gc leak on map change.
+		FSequenceRecorder::Get().ClearQueuedRecordings();
+		ListView->RebuildList();
+
+		// We also want to construct a new mutable default so it resets the recording paths to the
+		// default paths for the new map.
+		RecordingGroupDetailsView->SetObject(GetMutableDefault<USequenceRecorderActorGroup>());
+	}
 }
 
 void SSequenceRecorder::BindCommands()
@@ -332,6 +386,25 @@ void SSequenceRecorder::BindCommands()
 		FExecuteAction::CreateSP(this, &SSequenceRecorder::HandleRemoveAllRecordings),
 		FCanExecuteAction::CreateSP(this, &SSequenceRecorder::CanRemoveAllRecordings)
 		);
+
+	CommandList->MapAction(FSequenceRecorderCommands::Get().AddRecordingGroup,
+		FExecuteAction::CreateSP(this, &SSequenceRecorder::HandleAddRecordingGroup),
+		FCanExecuteAction::CreateSP(this, &SSequenceRecorder::CanAddRecordingGroup)
+	);
+
+	CommandList->MapAction(FSequenceRecorderCommands::Get().RemoveRecordingGroup,
+		FExecuteAction::CreateSP(this, &SSequenceRecorder::HandleRemoveRecordingGroup),
+		FCanExecuteAction::CreateSP(this, &SSequenceRecorder::CanRemoveRecordingGroup)
+	);
+
+	CommandList->MapAction(FSequenceRecorderCommands::Get().DuplicateRecordingGroup,
+		FExecuteAction::CreateSP(this, &SSequenceRecorder::HandleDuplicateRecordingGroup),
+		FCanExecuteAction::CreateSP(this, &SSequenceRecorder::CanDuplicateRecordingGroup)
+	);
+
+	// Append to level editor module so that shortcuts are accessible in level editor
+	FLevelEditorModule& LevelEditorModule = FModuleManager::GetModuleChecked<FLevelEditorModule>(TEXT("LevelEditor"));
+	LevelEditorModule.GetGlobalLevelEditorActions()->Append(CommandList.ToSharedRef());
 }
 
 TSharedRef<ITableRow> SSequenceRecorder::MakeListViewWidget(UActorRecording* Recording, const TSharedRef<STableViewBase>& OwnerTable) const
@@ -402,6 +475,14 @@ void SSequenceRecorder::HandleRemoveRecording()
 	{
 		FSequenceRecorder::Get().RemoveQueuedRecording(SelectedRecording);
 
+		// Remove the recording from the current group here. We can't use the
+		// FSequenceRecorder function as they are called when switching groups,
+		// and not just when the user removes items.
+		if (FSequenceRecorder::Get().GetRecordingGroup().IsValid())
+		{
+			FSequenceRecorder::Get().GetRecordingGroup()->RecordedActors.Remove(SelectedRecording);
+		}
+
 		TArray<TWeakObjectPtr<UObject>> SelectedObjects = ActorRecordingDetailsView->GetSelectedObjects();
 		if(SelectedObjects.Num() > 0 && SelectedObjects[0].Get() == SelectedRecording)
 		{
@@ -418,6 +499,10 @@ bool SSequenceRecorder::CanRemoveRecording() const
 void SSequenceRecorder::HandleRemoveAllRecordings()
 {
 	FSequenceRecorder::Get().ClearQueuedRecordings();
+	if (FSequenceRecorder::Get().GetRecordingGroup().IsValid())
+	{
+		FSequenceRecorder::Get().GetRecordingGroup()->RecordedActors.Empty();
+	}
 	ActorRecordingDetailsView->SetObject(nullptr);
 }
 
@@ -431,13 +516,87 @@ EActiveTimerReturnType SSequenceRecorder::HandleRefreshItems(double InCurrentTim
 	if(FSequenceRecorder::Get().AreQueuedRecordingsDirty())
 	{
 		ListView->RequestListRefresh();
-
 		FSequenceRecorder::Get().ResetQueuedRecordingsDirty();
 	}
 
 	return EActiveTimerReturnType::Continue;
 }
 
+void SSequenceRecorder::HandleAddRecordingGroup()
+{
+	FSequenceRecorder::Get().AddRecordingGroup();
+	check(FSequenceRecorder::Get().GetRecordingGroup().IsValid());
+	RecordingGroupDetailsView->SetObject(FSequenceRecorder::Get().GetRecordingGroup().Get());
+}
+
+bool SSequenceRecorder::CanAddRecordingGroup() const
+{
+	return !FSequenceRecorder::Get().IsRecording();
+}
+
+void SSequenceRecorder::HandleLoadRecordingActorGroup(FName Name)
+{
+	FSequenceRecorder::Get().LoadRecordingGroup(Name);
+
+	// Bind our details view to the newly loaded group.
+	if (FSequenceRecorder::Get().GetRecordingGroup().IsValid())
+	{
+		RecordingGroupDetailsView->SetObject(FSequenceRecorder::Get().GetRecordingGroup().Get());
+	}
+	else
+	{
+		// If they've loaded the "None" profile we create a new default as well to reset the paths.
+		RecordingGroupDetailsView->SetObject(GetMutableDefault<USequenceRecorderActorGroup>());
+	}
+}
+
+void SSequenceRecorder::HandleRemoveRecordingGroup()
+{
+	FSequenceRecorder::Get().RemoveCurrentRecordingGroup();
+
+	// See if there's any recordings left, if so we'll load the last one, otherwise we need to load
+	// a default so that the UI is still visible.
+	TArray<FName> RecordingProfiles = FSequenceRecorder::Get().GetRecordingGroupNames();
+	if (RecordingProfiles.Num() > 0)
+	{
+		FSequenceRecorder::Get().LoadRecordingGroup(RecordingProfiles[RecordingProfiles.Num() - 1]);
+		check(FSequenceRecorder::Get().GetRecordingGroup().Get());
+
+		RecordingGroupDetailsView->SetObject(FSequenceRecorder::Get().GetRecordingGroup().Get());
+	}
+	else
+	{
+		RecordingGroupDetailsView->SetObject(GetMutableDefault<USequenceRecorderActorGroup>());
+	}
+}
+
+bool SSequenceRecorder::CanRemoveRecordingGroup() const
+{
+	TWeakObjectPtr<USequenceRecorderActorGroup> RecordingGroup = FSequenceRecorder::Get().GetRecordingGroup();
+	if (RecordingGroup.IsValid())
+	{
+		return RecordingGroup->GroupName != NAME_None;
+	}
+
+	return false;
+}
+
+void SSequenceRecorder::HandleDuplicateRecordingGroup()
+{
+	FSequenceRecorder::Get().DuplicateRecordingGroup();
+}
+
+bool SSequenceRecorder::CanDuplicateRecordingGroup() const
+{
+	TWeakObjectPtr<USequenceRecorderActorGroup> RecordingGroup = FSequenceRecorder::Get().GetRecordingGroup();
+	if (RecordingGroup.IsValid())
+	{
+		return RecordingGroup->GroupName != NAME_None;
+	}
+
+	return false;
+}
+
 TOptional<float> SSequenceRecorder::GetDelayPercent() const
 {
 	const float Delay = GetDefault<USequenceRecorderSettings>()->RecordingDelay;
diff --git a/Engine/Source/Editor/SequenceRecorder/Private/SSequenceRecorder.h b/Engine/Source/Editor/SequenceRecorder/Private/SSequenceRecorder.h
--- a/Engine/Source/Editor/SequenceRecorder/Private/SSequenceRecorder.h
+++ b/Engine/Source/Editor/SequenceRecorder/Private/SSequenceRecorder.h
@@ -9,12 +9,17 @@
 #include "Widgets/Views/STableViewBase.h"
 #include "Widgets/Views/STableRow.h"
 #include "ActorRecording.h"
+#include "PropertyEditorDelegates.h"
+#include "ISinglePropertyView.h"
+#include "IStructureDetailsView.h"
 
 class FActiveTimerHandle;
 class FUICommandList;
+class FUICommandInfo;
 class IDetailsView;
 class SProgressBar;
 class FDragDropOperation;
+class SEditableTextBox;
 
 class SSequenceRecorder : public SCompoundWidget
 {
@@ -25,6 +30,9 @@ class SSequenceRecorder : public SCompoundWidget
 	SLATE_END_ARGS()
 
 	void Construct(const FArguments& Args);
+	
+	/** SSequenceRecorder destructor */
+	~SSequenceRecorder();
 
 private:
 
@@ -62,8 +70,25 @@ private:
 
 	bool CanRemoveAllRecordings() const;
 
+	void HandleAddRecordingGroup();
+
+
+	bool CanAddRecordingGroup() const;
+
+	void HandleRemoveRecordingGroup();
+
+	bool CanRemoveRecordingGroup() const;
+
+	void HandleDuplicateRecordingGroup();
+
+	bool CanDuplicateRecordingGroup() const;
+
+	void HandleRecordingProfileNameCommitted(const FText& InText, ETextCommit::Type InCommitType);
+
 	EActiveTimerReturnType HandleRefreshItems(double InCurrentTime, float InDeltaTime);
 
+	void HandleMapUnload(UObject* Object);
+
 	TOptional<float> GetDelayPercent() const;
 
 	void OnDelayChanged(float NewValue);
@@ -76,11 +101,24 @@ private:
 
 	bool OnRecordingListAllowDrop( TSharedPtr<FDragDropOperation> DragDropOperation );
 
+public:
+	TSharedPtr<FUICommandList> GetCommandList() const
+	{
+		return CommandList;
+	}
+
+	void HandleLoadRecordingActorGroup(FName Name);
+
 private:
+	/** This is the Detail View for the USequenceRecorderSettings */
 	TSharedPtr<IDetailsView> SequenceRecordingDetailsView;
 
+	/** This is the Detail View for the currently selected UActorRecording */
 	TSharedPtr<IDetailsView> ActorRecordingDetailsView;
 
+	/** This is the Detail View for the currently selected USequenceActorGroup */
+	TSharedPtr<IDetailsView> RecordingGroupDetailsView;
+
 	TSharedPtr<SListView<UActorRecording*>> ListView;
 
 	TSharedPtr<FUICommandList> CommandList;
diff --git a/Engine/Source/Editor/SequenceRecorder/Private/Sections/MovieScene3DTransformSectionRecorder.cpp b/Engine/Source/Editor/SequenceRecorder/Private/Sections/MovieScene3DTransformSectionRecorder.cpp
--- a/Engine/Source/Editor/SequenceRecorder/Private/Sections/MovieScene3DTransformSectionRecorder.cpp
+++ b/Engine/Source/Editor/SequenceRecorder/Private/Sections/MovieScene3DTransformSectionRecorder.cpp
@@ -8,6 +8,8 @@
 #include "Tracks/MovieScene3DTransformTrack.h"
 #include "SequenceRecorder.h"
 #include "SequenceRecorderSettings.h"
+#include "Algo/Transform.h"
+#include "Channels/MovieSceneChannelProxy.h"
 
 TSharedPtr<IMovieSceneSectionRecorder> FMovieScene3DTransformSectionRecorderFactory::CreateSectionRecorder(const FActorRecordingSettings& InActorRecordingSettings) const
 {
@@ -45,7 +47,8 @@ void FMovieScene3DTransformSectionRecorder::CreateSection(UObject* InObjectToRec
 	ObjectToRecord = InObjectToRecord;
 	Guid = InGuid;
 	bWasAttached = false;
-		
+	RecordingStartTime = Time;
+
 	MovieScene = InMovieScene;
 	MovieSceneTrack = InMovieScene->AddTrack<UMovieScene3DTransformTrack>(Guid);
 	if(MovieSceneTrack.IsValid())
@@ -54,40 +57,50 @@ void FMovieScene3DTransformSectionRecorder::CreateSection(UObject* InObjectToRec
 
 		MovieSceneTrack->AddSection(*MovieSceneSection);
 
-		const bool bUnwindRotation = false;
-
 		DefaultTransform = FTransform::Identity;
 		GetTransformToRecord(DefaultTransform);
-		FVector EulerRotation = DefaultTransform.GetRotation().Rotator().Euler();
-
-		MovieSceneSection->SetDefault(FTransformKey(EKey3DTransformChannel::Translation, EAxis::X, DefaultTransform.GetTranslation().X, bUnwindRotation));
-		MovieSceneSection->SetDefault(FTransformKey(EKey3DTransformChannel::Translation, EAxis::Y, DefaultTransform.GetTranslation().Y, bUnwindRotation));
-		MovieSceneSection->SetDefault(FTransformKey(EKey3DTransformChannel::Translation, EAxis::Z, DefaultTransform.GetTranslation().Z, bUnwindRotation));
-
-		MovieSceneSection->SetDefault(FTransformKey(EKey3DTransformChannel::Rotation, EAxis::X, EulerRotation.X, bUnwindRotation));
-		MovieSceneSection->SetDefault(FTransformKey(EKey3DTransformChannel::Rotation, EAxis::Y, EulerRotation.Y, bUnwindRotation));
-		MovieSceneSection->SetDefault(FTransformKey(EKey3DTransformChannel::Rotation, EAxis::Z, EulerRotation.Z, bUnwindRotation));
 
-		MovieSceneSection->SetDefault(FTransformKey(EKey3DTransformChannel::Scale, EAxis::X, DefaultTransform.GetScale3D().X, bUnwindRotation));
-		MovieSceneSection->SetDefault(FTransformKey(EKey3DTransformChannel::Scale, EAxis::Y, DefaultTransform.GetScale3D().Y, bUnwindRotation));
-		MovieSceneSection->SetDefault(FTransformKey(EKey3DTransformChannel::Scale, EAxis::Z, DefaultTransform.GetScale3D().Z, bUnwindRotation));
-
-		MovieSceneSection->SetStartTime(Time);
-		MovieSceneSection->SetIsInfinite(true);
+		FVector Translation   = DefaultTransform.GetTranslation();
+		FVector EulerRotation = DefaultTransform.GetRotation().Rotator().Euler();
+		FVector Scale         = DefaultTransform.GetScale3D();
+
+		TArrayView<FMovieSceneFloatChannel*> FloatChannels = MovieSceneSection->GetChannelProxy().GetChannels<FMovieSceneFloatChannel>();
+		FloatChannels[0]->SetDefault(Translation.X);
+		FloatChannels[1]->SetDefault(Translation.Y);
+		FloatChannels[2]->SetDefault(Translation.Z);
+		FloatChannels[3]->SetDefault(EulerRotation.X);
+		FloatChannels[4]->SetDefault(EulerRotation.Y);
+		FloatChannels[5]->SetDefault(EulerRotation.Z);
+		FloatChannels[6]->SetDefault(Scale.X);
+		FloatChannels[7]->SetDefault(Scale.Y);
+		FloatChannels[8]->SetDefault(Scale.Z);
+
+		MovieSceneSection->SetRange(TRange<FFrameNumber>::All());
 	}
 }
 
 void FMovieScene3DTransformSectionRecorder::FinalizeSection()
 {
+	bool bWasRecording = bRecording;
+	bRecording = false;
+
 	FScopedSlowTask SlowTask(4.0f, NSLOCTEXT("SequenceRecorder", "ProcessingTransforms", "Processing Transforms"));
 
-	bRecording = false;
+	check (	BufferedTransforms.Times.Num() == BufferedTransforms.LocationX.Num());	
+	check (	BufferedTransforms.Times.Num() == BufferedTransforms.LocationY.Num());
+	check (	BufferedTransforms.Times.Num() == BufferedTransforms.LocationZ.Num());
+	check (	BufferedTransforms.Times.Num() == BufferedTransforms.RotationX.Num());
+	check (	BufferedTransforms.Times.Num() == BufferedTransforms.RotationY.Num());
+	check (	BufferedTransforms.Times.Num() == BufferedTransforms.RotationZ.Num());
+	check (	BufferedTransforms.Times.Num() == BufferedTransforms.ScaleX.Num());
+	check (	BufferedTransforms.Times.Num() == BufferedTransforms.ScaleY.Num());
+	check (	BufferedTransforms.Times.Num() == BufferedTransforms.ScaleZ.Num());
 
 	// if we have a valid animation recorder, we need to build our transforms from the animation
 	// so we properly synchronize our keyframes
-	if(AnimRecorder.IsValid())
+	if(AnimRecorder.IsValid() && bWasRecording)
 	{
-		check(BufferedTransforms.Num() == 0);
+		check(BufferedTransforms.Times.Num() == 0);
 
 		UAnimSequence* AnimSequence = AnimRecorder->GetAnimSequence();
 		USkeletalMeshComponent* SkeletalMeshComponent = AnimRecorder->GetSkeletalMeshComponent();
@@ -118,7 +131,8 @@ void FMovieScene3DTransformSectionRecorder::FinalizeSection()
 
 				check(RootIndex != INDEX_NONE);
 
-				const float StartTime = MovieSceneSection->GetStartTime();
+				FFrameRate FrameResolution = MovieSceneSection->GetTypedOuter<UMovieScene>()->GetFrameResolution();
+				const FFrameNumber StartTime = (RecordingStartTime * FrameResolution).FloorToFrame();
 
 				// we may need to offset the transform here if the animation was not recorded on the root component
 				FTransform InvComponentTransform = AnimRecorder->GetComponentTransform().Inverse();
@@ -155,7 +169,8 @@ void FMovieScene3DTransformSectionRecorder::FinalizeSection()
 						Transform.SetScale3D(RawTrack.ScaleKeys[0]);
 					}
 
-					BufferedTransforms.Add(FBufferedTransformKey(InvComponentTransform * Transform, StartTime + AnimSequence->GetTimeAtFrame(KeyIndex)));
+					FFrameNumber AnimationFrame = (AnimSequence->GetTimeAtFrame(KeyIndex) * FrameResolution).FloorToFrame();
+					BufferedTransforms.Add(InvComponentTransform * Transform, StartTime + AnimationFrame);
 				}
 			}
 		}
@@ -168,98 +183,90 @@ void FMovieScene3DTransformSectionRecorder::FinalizeSection()
 	// - Net quantize may use quaternions.
 	// - Scene components cache transforms as quaternions.
 	// - Gameplay is free to clamp/fmod rotations as it sees fit.
-	int32 TransformCount = BufferedTransforms.Num();
+	int32 TransformCount = BufferedTransforms.Times.Num();
 	for(int32 TransformIndex = 0; TransformIndex < TransformCount - 1; TransformIndex++)
 	{
-		FRotator& Rotator = BufferedTransforms[TransformIndex].WoundRotation;
-		FRotator& NextRotator = BufferedTransforms[TransformIndex + 1].WoundRotation;
-
-		FMath::WindRelativeAnglesDegrees(Rotator.Pitch, NextRotator.Pitch);
-		FMath::WindRelativeAnglesDegrees(Rotator.Yaw, NextRotator.Yaw);
-		FMath::WindRelativeAnglesDegrees(Rotator.Roll, NextRotator.Roll);
+		FMath::WindRelativeAnglesDegrees(BufferedTransforms.RotationZ[TransformIndex], BufferedTransforms.RotationZ[TransformIndex+1]);
+		FMath::WindRelativeAnglesDegrees(BufferedTransforms.RotationY[TransformIndex], BufferedTransforms.RotationY[TransformIndex+1]);
+		FMath::WindRelativeAnglesDegrees(BufferedTransforms.RotationX[TransformIndex], BufferedTransforms.RotationX[TransformIndex+1]);
 	}
 
 	SlowTask.EnterProgressFrame();
 
-	// never unwind rotations
-	const bool bUnwindRotation = false;
 	// If we are syncing to an animation, use linear interpolation to avoid foot sliding etc. 
 	// Otherwise use cubic for better quality (much better for projectiles etc.)
-	const EMovieSceneKeyInterpolation Interpolation = AnimRecorder.IsValid() ? EMovieSceneKeyInterpolation::Linear : EMovieSceneKeyInterpolation::Auto;
+
+	// @todo: sequencer-timecode: this was previously never actually used - should it be??
+	const ERichCurveInterpMode Interpolation = AnimRecorder.IsValid() ? RCIM_Linear : RCIM_Cubic;
 
 	// add buffered transforms
-	TArray<FRichCurveKey> TranslationXKeys; TranslationXKeys.SetNum(BufferedTransforms.Num());
-	TArray<FRichCurveKey> TranslationYKeys; TranslationYKeys.SetNum(BufferedTransforms.Num());
-	TArray<FRichCurveKey> TranslationZKeys; TranslationZKeys.SetNum(BufferedTransforms.Num());
-	TArray<FRichCurveKey> RotationXKeys; RotationXKeys.SetNum(BufferedTransforms.Num());
-	TArray<FRichCurveKey> RotationYKeys; RotationYKeys.SetNum(BufferedTransforms.Num());
-	TArray<FRichCurveKey> RotationZKeys; RotationZKeys.SetNum(BufferedTransforms.Num());
-	TArray<FRichCurveKey> ScaleXKeys; ScaleXKeys.SetNum(BufferedTransforms.Num());
-	TArray<FRichCurveKey> ScaleYKeys; ScaleYKeys.SetNum(BufferedTransforms.Num());
-	TArray<FRichCurveKey> ScaleZKeys; ScaleZKeys.SetNum(BufferedTransforms.Num());
-
-	int32 Index = 0;
-	for(const FBufferedTransformKey& BufferedTransform : BufferedTransforms)
+	TArrayView<FMovieSceneFloatChannel*> FloatChannels = MovieSceneSection->GetChannelProxy().GetChannels<FMovieSceneFloatChannel>();
+
+	auto Transformation = [Interpolation](float In)
 	{
-		const FVector Translation = BufferedTransform.Transform.GetTranslation();
-		const FVector Rotation = BufferedTransform.WoundRotation.Euler();
-		const FVector Scale = BufferedTransform.Transform.GetScale3D();
+		FMovieSceneFloatValue NewValue(In);
+		NewValue.InterpMode = Interpolation;
+		return NewValue;
+	};
+	TArray<FMovieSceneFloatValue> FloatValues;
 
-		TranslationXKeys[Index] = FRichCurveKey(BufferedTransform.KeyTime, Translation.X);
-		TranslationYKeys[Index] = FRichCurveKey(BufferedTransform.KeyTime, Translation.Y);
-		TranslationZKeys[Index] = FRichCurveKey(BufferedTransform.KeyTime, Translation.Z);
+	FloatValues.Reset(BufferedTransforms.LocationX.Num());
+	Algo::Transform(BufferedTransforms.LocationX, FloatValues, Transformation);
+	FloatChannels[0]->Set(BufferedTransforms.Times, MoveTemp(FloatValues));
 
-		RotationXKeys[Index] = FRichCurveKey(BufferedTransform.KeyTime, Rotation.X);
-		RotationYKeys[Index] = FRichCurveKey(BufferedTransform.KeyTime, Rotation.Y);
-		RotationZKeys[Index] = FRichCurveKey(BufferedTransform.KeyTime, Rotation.Z);
+	FloatValues.Reset(BufferedTransforms.LocationY.Num());
+	Algo::Transform(BufferedTransforms.LocationY, FloatValues, Transformation);
+	FloatChannels[1]->Set(BufferedTransforms.Times, MoveTemp(FloatValues));
 
-		ScaleXKeys[Index] = FRichCurveKey(BufferedTransform.KeyTime, Scale.X);
-		ScaleYKeys[Index] = FRichCurveKey(BufferedTransform.KeyTime, Scale.Y);
-		ScaleZKeys[Index] = FRichCurveKey(BufferedTransform.KeyTime, Scale.Z);
+	FloatValues.Reset(BufferedTransforms.LocationZ.Num());
+	Algo::Transform(BufferedTransforms.LocationZ, FloatValues, Transformation);
+	FloatChannels[2]->Set(BufferedTransforms.Times, MoveTemp(FloatValues));
 
-		++Index;
-	}
-	
-	MovieSceneSection->GetTranslationCurve(EAxis::X).SetKeys(TranslationXKeys);
-	MovieSceneSection->GetTranslationCurve(EAxis::Y).SetKeys(TranslationYKeys);
-	MovieSceneSection->GetTranslationCurve(EAxis::Z).SetKeys(TranslationZKeys);
-	MovieSceneSection->GetRotationCurve(EAxis::X).SetKeys(RotationXKeys);
-	MovieSceneSection->GetRotationCurve(EAxis::Y).SetKeys(RotationYKeys);
-	MovieSceneSection->GetRotationCurve(EAxis::Z).SetKeys(RotationZKeys);
-	MovieSceneSection->GetScaleCurve(EAxis::X).SetKeys(ScaleXKeys);
-	MovieSceneSection->GetScaleCurve(EAxis::Y).SetKeys(ScaleYKeys);
-	MovieSceneSection->GetScaleCurve(EAxis::Z).SetKeys(ScaleZKeys);
+	FloatValues.Reset(BufferedTransforms.RotationX.Num());
+	Algo::Transform(BufferedTransforms.RotationX, FloatValues, Transformation);
+	FloatChannels[3]->Set(BufferedTransforms.Times, MoveTemp(FloatValues));
+
+	FloatValues.Reset(BufferedTransforms.RotationY.Num());
+	Algo::Transform(BufferedTransforms.RotationY, FloatValues, Transformation);
+	FloatChannels[4]->Set(BufferedTransforms.Times, MoveTemp(FloatValues));
+
+	FloatValues.Reset(BufferedTransforms.RotationZ.Num());
+	Algo::Transform(BufferedTransforms.RotationZ, FloatValues, Transformation);
+	FloatChannels[5]->Set(BufferedTransforms.Times, MoveTemp(FloatValues));
+
+	FloatValues.Reset(BufferedTransforms.ScaleX.Num());
+	Algo::Transform(BufferedTransforms.ScaleX, FloatValues, Transformation);
+	FloatChannels[6]->Set(BufferedTransforms.Times, MoveTemp(FloatValues));
+
+	FloatValues.Reset(BufferedTransforms.ScaleY.Num());
+	Algo::Transform(BufferedTransforms.ScaleY, FloatValues, Transformation);
+	FloatChannels[7]->Set(BufferedTransforms.Times, MoveTemp(FloatValues));
+
+	FloatValues.Reset(BufferedTransforms.ScaleZ.Num());
+	Algo::Transform(BufferedTransforms.ScaleZ, FloatValues, Transformation);
+	FloatChannels[8]->Set(BufferedTransforms.Times, MoveTemp(FloatValues));
 
 	FTransform FirstTransform = FTransform::Identity;
-	if (BufferedTransforms.Num())
+	if (BufferedTransforms.Times.Num())
 	{
-		FirstTransform = BufferedTransforms[0].Transform;
+		FirstTransform.SetTranslation(FVector(BufferedTransforms.LocationX[0], BufferedTransforms.LocationY[0], BufferedTransforms.LocationZ[0]));
+		FirstTransform.SetRotation(FQuat(FRotator(BufferedTransforms.RotationX[0], BufferedTransforms.RotationY[0], BufferedTransforms.RotationZ[0])));
+		FirstTransform.SetScale3D(FVector(BufferedTransforms.ScaleX[0], BufferedTransforms.ScaleY[0], BufferedTransforms.ScaleZ[0]));
 	}
 
-	BufferedTransforms.Empty();
+	BufferedTransforms = FBufferedTransformKeys();
 
 	SlowTask.EnterProgressFrame();
 
 	// now remove linear keys
-	const TPair<FRichCurve*, float> CurvesAndTolerances[] =
-	{
-		TPair<FRichCurve*, float>(&MovieSceneSection->GetTranslationCurve(EAxis::X), KINDA_SMALL_NUMBER),
-		TPair<FRichCurve*, float>(&MovieSceneSection->GetTranslationCurve(EAxis::Y), KINDA_SMALL_NUMBER),
-		TPair<FRichCurve*, float>(&MovieSceneSection->GetTranslationCurve(EAxis::Z), KINDA_SMALL_NUMBER),
-		TPair<FRichCurve*, float>(&MovieSceneSection->GetRotationCurve(EAxis::X), KINDA_SMALL_NUMBER),
-		TPair<FRichCurve*, float>(&MovieSceneSection->GetRotationCurve(EAxis::Y), KINDA_SMALL_NUMBER),
-		TPair<FRichCurve*, float>(&MovieSceneSection->GetRotationCurve(EAxis::Z), KINDA_SMALL_NUMBER),
-		TPair<FRichCurve*, float>(&MovieSceneSection->GetScaleCurve(EAxis::X), KINDA_SMALL_NUMBER),
-		TPair<FRichCurve*, float>(&MovieSceneSection->GetScaleCurve(EAxis::Y), KINDA_SMALL_NUMBER),
-		TPair<FRichCurve*, float>(&MovieSceneSection->GetScaleCurve(EAxis::Z), KINDA_SMALL_NUMBER),
-	};
-
 	const USequenceRecorderSettings* Settings = GetDefault<USequenceRecorderSettings>();
-	if (Settings->bReduceKeys) 
+	if (Settings->bReduceKeys)
 	{
-		for (const TPair<FRichCurve*, float>& CurveAndTolerance : CurvesAndTolerances)
+		FKeyDataOptimizationParams Params;
+
+		for (FMovieSceneFloatChannel* Channel : FloatChannels)
 		{
-			CurveAndTolerance.Key->RemoveRedundantKeys(CurveAndTolerance.Value);
+			MovieScene::Optimize(Channel, Params);
 		}
 	}
 
@@ -267,16 +274,24 @@ void FMovieScene3DTransformSectionRecorder::FinalizeSection()
 	// transform tracks. Without this track, relative transforms would accumulate.
 	if(!bWasAttached)
 	{
-		// now we have reduced our keys, if we dont have any, remove the section as it is redundant
-		if( MovieSceneSection->GetTranslationCurve(EAxis::X).GetNumKeys() == 0 && 
-			MovieSceneSection->GetTranslationCurve(EAxis::Y).GetNumKeys() == 0 &&
-			MovieSceneSection->GetTranslationCurve(EAxis::Z).GetNumKeys() == 0 &&
-			MovieSceneSection->GetRotationCurve(EAxis::X).GetNumKeys() == 0 &&
-			MovieSceneSection->GetRotationCurve(EAxis::Y).GetNumKeys() == 0 &&
-			MovieSceneSection->GetRotationCurve(EAxis::Z).GetNumKeys() == 0 &&
-			MovieSceneSection->GetScaleCurve(EAxis::X).GetNumKeys() == 0 &&
-			MovieSceneSection->GetScaleCurve(EAxis::Y).GetNumKeys() == 0 &&
-			MovieSceneSection->GetScaleCurve(EAxis::Z).GetNumKeys() == 0)
+		bool bCanRemoveTrack = true;
+		for (FMovieSceneFloatChannel* Channel : MovieSceneSection->GetChannelProxy().GetChannels<FMovieSceneFloatChannel>())
+		{
+			if (Channel)
+			{
+				int32 NumKeys = Channel->GetTimes().Num();
+				if (NumKeys == 1)
+				{
+					*Channel = FMovieSceneFloatChannel();
+				}
+				else if (NumKeys > 1)
+				{
+					bCanRemoveTrack = false;
+				}
+			}
+		}
+
+		if (bCanRemoveTrack)
 		{
 			if(DefaultTransform.Equals(FTransform::Identity))
 			{
@@ -311,7 +326,8 @@ void FMovieScene3DTransformSectionRecorder::Record(float CurrentTime)
 			}
 		}
 
-		MovieSceneSection->SetEndTime(CurrentTime);
+		FFrameRate   FrameResolution = MovieSceneSection->GetTypedOuter<UMovieScene>()->GetFrameResolution();
+		FFrameNumber CurrentFrame    = (CurrentTime * FrameResolution).FloorToFrame();
 
 		if(bRecording)
 		{
@@ -321,7 +337,7 @@ void FMovieScene3DTransformSectionRecorder::Record(float CurrentTime)
 				FTransform TransformToRecord;
 				if (GetTransformToRecord(TransformToRecord))
 				{
-					BufferedTransforms.Add(FBufferedTransformKey(TransformToRecord, CurrentTime));
+					BufferedTransforms.Add(TransformToRecord, CurrentFrame);
 				}
 			}
 		}
diff --git a/Engine/Source/Editor/SequenceRecorder/Private/Sections/MovieScene3DTransformSectionRecorder.h b/Engine/Source/Editor/SequenceRecorder/Private/Sections/MovieScene3DTransformSectionRecorder.h
--- a/Engine/Source/Editor/SequenceRecorder/Private/Sections/MovieScene3DTransformSectionRecorder.h
+++ b/Engine/Source/Editor/SequenceRecorder/Private/Sections/MovieScene3DTransformSectionRecorder.h
@@ -20,7 +20,7 @@ public:
 	virtual ~FMovieScene3DTransformSectionRecorderFactory() {}
 
 	virtual bool CanRecordObject(class UObject* InObjectToRecord) const override;
-	virtual UObject* CreateSettingsObject() const override { return NewObject<UMovieScene3DTransformSectionRecorderSettings>(); }
+	virtual UObject* CreateSettingsObject(class UObject* InOuter) const override { return NewObject<UMovieScene3DTransformSectionRecorderSettings>(InOuter, FName(TEXT("MovieScene3DTransformSectionRecorderSettings"))); }
 
 	TSharedPtr<class FMovieScene3DTransformSectionRecorder> CreateSectionRecorder(bool bRecordTransforms, TSharedPtr<class FMovieSceneAnimationSectionRecorder> InAnimRecorder) const;
 
@@ -29,18 +29,29 @@ private:
 };
 
 /** Structure used to buffer up transform keys. Keys are inserted into tracks in FinalizeSection() */
-struct FBufferedTransformKey
+struct FBufferedTransformKeys
 {
-	FBufferedTransformKey(const FTransform& InTransform, float InKeyTime)
-		: Transform(InTransform)
-		, WoundRotation(InTransform.Rotator())
-		, KeyTime(InKeyTime)
+	void Add(const FTransform& InTransform, FFrameNumber InKeyTime)
 	{
+		Times.Add(InKeyTime);
+		LocationX.Add(InTransform.GetTranslation().X);
+		LocationY.Add(InTransform.GetTranslation().Y);
+		LocationZ.Add(InTransform.GetTranslation().Z);
+
+		FRotator WoundRoation = InTransform.Rotator();
+		RotationX.Add(WoundRoation.Roll);
+		RotationY.Add(WoundRoation.Pitch);
+		RotationZ.Add(WoundRoation.Yaw);
+
+		ScaleX.Add(InTransform.GetScale3D().X);
+		ScaleY.Add(InTransform.GetScale3D().Y);
+		ScaleZ.Add(InTransform.GetScale3D().Z);
 	}
 
-	FTransform Transform;
-	FRotator WoundRotation;
-	float KeyTime;
+	TArray<FFrameNumber> Times;
+	TArray<float> LocationX, LocationY, LocationZ;
+	TArray<float> RotationX, RotationY, RotationZ;
+	TArray<float> ScaleX, ScaleY, ScaleZ;
 };
 
 class FMovieScene3DTransformSectionRecorder : public IMovieSceneSectionRecorder
@@ -82,7 +93,7 @@ private:
 	TWeakObjectPtr<class UMovieScene3DTransformSection> MovieSceneSection;
 
 	/** Buffer of transform keys. Keys are inserted into tracks in FinalizeSection() */
-	TArray<FBufferedTransformKey> BufferedTransforms;
+	FBufferedTransformKeys BufferedTransforms;
 
 	/** Flag if we are actually recording or not */
 	bool bRecording;
@@ -98,4 +109,7 @@ private:
 
 	/** The guid being recorded to */
 	FGuid Guid;
+
+	/** The guid being recorded to */
+	float RecordingStartTime;
 };
diff --git a/Engine/Source/Editor/SequenceRecorder/Private/Sections/MovieSceneAnimationSectionRecorder.cpp b/Engine/Source/Editor/SequenceRecorder/Private/Sections/MovieSceneAnimationSectionRecorder.cpp
--- a/Engine/Source/Editor/SequenceRecorder/Private/Sections/MovieSceneAnimationSectionRecorder.cpp
+++ b/Engine/Source/Editor/SequenceRecorder/Private/Sections/MovieSceneAnimationSectionRecorder.cpp
@@ -9,6 +9,7 @@
 #include "SequenceRecorderUtils.h"
 #include "SequenceRecorderSettings.h"
 #include "ActorRecording.h"
+#include "SequenceRecorder.h"
 
 TSharedPtr<IMovieSceneSectionRecorder> FMovieSceneAnimationSectionRecorderFactory::CreateSectionRecorder(const FActorRecordingSettings& InActorRecordingSettings) const
 {
@@ -71,13 +72,14 @@ void FMovieSceneAnimationSectionRecorder::CreateSection(UObject* InObjectToRecor
 				// build an asset path
 				const USequenceRecorderSettings* Settings = GetDefault<USequenceRecorderSettings>();
 
-				FString AssetPath = Settings->SequenceRecordingBasePath.Path;
+				FString AssetPath = FSequenceRecorder::Get().GetSequenceRecordingBasePath();
 				if (Settings->AnimationSubDirectory.Len() > 0)
 				{
 					AssetPath /= Settings->AnimationSubDirectory;
 				}
 
-				FString AssetName = Settings->SequenceName.Len() > 0 ? Settings->SequenceName : TEXT("RecordedSequence");
+				const FString SequenceName = FSequenceRecorder::Get().GetSequenceRecordingName();
+				FString AssetName = SequenceName.Len() > 0 ? SequenceName : TEXT("RecordedSequence");
 				AssetName += TEXT("_");
 				check(Actor);
 				AssetName += Actor->GetActorLabel();
@@ -101,7 +103,10 @@ void FMovieSceneAnimationSectionRecorder::CreateSection(UObject* InObjectToRecor
 					UMovieSceneSkeletalAnimationTrack* AnimTrack = MovieScene->AddTrack<UMovieSceneSkeletalAnimationTrack>(Guid);
 					if (AnimTrack)
 					{
-						AnimTrack->AddNewAnimation(Time, AnimSequence.Get());
+						FFrameRate   FrameResolution = AnimTrack->GetTypedOuter<UMovieScene>()->GetFrameResolution();
+						FFrameNumber CurrentFrame    = (Time * FrameResolution).FloorToFrame();
+
+						AnimTrack->AddNewAnimation(CurrentFrame, AnimSequence.Get());
 						MovieSceneSection = Cast<UMovieSceneSkeletalAnimationSection>(AnimTrack->GetAllSections()[0]);
 					}
 				}
@@ -128,9 +133,12 @@ void FMovieSceneAnimationSectionRecorder::FinalizeSection()
 		FAnimationRecorderManager::Get().StopRecordingAnimation(SkeletalMeshComponent.Get(), bShowMessage);
 	}
 
-	if(MovieSceneSection.IsValid() && AnimSequence.IsValid())
+	if(MovieSceneSection.IsValid() && AnimSequence.IsValid() && MovieSceneSection->HasStartFrame())
 	{
-		MovieSceneSection->SetEndTime(MovieSceneSection->GetStartTime() + AnimSequence->GetPlayLength());
+		FFrameRate   FrameResolution = MovieSceneSection->GetTypedOuter<UMovieScene>()->GetFrameResolution();
+		FFrameNumber SequenceLength  = (AnimSequence->GetPlayLength() * FrameResolution).FloorToFrame();
+		
+		MovieSceneSection->SetEndFrame(TRangeBound<FFrameNumber>::Exclusive(MovieSceneSection->GetInclusiveStartFrame() + SequenceLength));
 	}
 }
 
diff --git a/Engine/Source/Editor/SequenceRecorder/Private/Sections/MovieScenePropertyRecorder.cpp b/Engine/Source/Editor/SequenceRecorder/Private/Sections/MovieScenePropertyRecorder.cpp
--- a/Engine/Source/Editor/SequenceRecorder/Private/Sections/MovieScenePropertyRecorder.cpp
+++ b/Engine/Source/Editor/SequenceRecorder/Private/Sections/MovieScenePropertyRecorder.cpp
@@ -14,6 +14,7 @@
 #include "Tracks/MovieSceneColorTrack.h"
 #include "Sections/MovieSceneVectorSection.h"
 #include "Tracks/MovieSceneVectorTrack.h"
+#include "Channels/MovieSceneChannelProxy.h"
 
 // current set of compiled-in property types
 
@@ -34,11 +35,16 @@ UMovieSceneSection* FMovieScenePropertyRecorder<bool>::AddSection(UObject* InObj
 			Track->SetPropertyNameAndPath(*Binding.GetProperty(*InObjectToRecord)->GetDisplayNameText().ToString(), Binding.GetPropertyPath());
 		}
 
-		UMovieSceneBoolSection* Section = Cast<UMovieSceneBoolSection>(Track->CreateNewSection());
-		Section->SetDefault(PreviousValue);
-		Section->SetStartTime(InTime);
-		Section->SetEndTime(InTime);
-		Section->AddKey(InTime, PreviousValue, EMovieSceneKeyInterpolation::Break);
+		UMovieSceneBoolSection* Section         = Cast<UMovieSceneBoolSection>(Track->CreateNewSection());
+
+		FFrameRate   FrameResolution = Section->GetTypedOuter<UMovieScene>()->GetFrameResolution();
+		FFrameNumber CurrentFrame    = (InTime * FrameResolution).FloorToFrame();
+
+		Section->SetRange(TRange<FFrameNumber>(CurrentFrame));
+
+		TArrayView<FMovieSceneFloatChannel*> FloatChannels = Section->GetChannelProxy().GetChannels<FMovieSceneFloatChannel>();
+		FloatChannels[0]->SetDefault(PreviousValue);
+		FloatChannels[0]->AddCubicKey(CurrentFrame, PreviousValue, RCTM_Break);
 
 		Track->AddSection(*Section);
 
@@ -51,7 +57,7 @@ UMovieSceneSection* FMovieScenePropertyRecorder<bool>::AddSection(UObject* InObj
 template <>
 void FMovieScenePropertyRecorder<bool>::AddKeyToSection(UMovieSceneSection* InSection, const FPropertyKey<bool>& InKey)
 {
-	CastChecked<UMovieSceneBoolSection>(InSection)->AddKey(InKey.Time, InKey.Value, EMovieSceneKeyInterpolation::Break);
+	InSection->GetChannelProxy().GetChannel<FMovieSceneBoolChannel>(0)->GetInterface().AddKey(InKey.Time, InKey.Value);
 }
 
 template <>
@@ -77,10 +83,15 @@ UMovieSceneSection* FMovieScenePropertyRecorder<uint8>::AddSection(UObject* InOb
 		}
 
 		UMovieSceneByteSection* Section = Cast<UMovieSceneByteSection>(Track->CreateNewSection());
-		Section->SetDefault(PreviousValue);
-		Section->SetStartTime(InTime);
-		Section->SetEndTime(InTime);
-		Section->AddKey(InTime, PreviousValue, EMovieSceneKeyInterpolation::Break);
+
+		FFrameRate   FrameResolution = Section->GetTypedOuter<UMovieScene>()->GetFrameResolution();
+		FFrameNumber CurrentFrame    = (InTime * FrameResolution).FloorToFrame();
+
+		Section->SetRange(TRange<FFrameNumber>::Inclusive(CurrentFrame, CurrentFrame));
+
+		FMovieSceneByteChannel* Channel = Section->GetChannelProxy().GetChannel<FMovieSceneByteChannel>(0);
+		Channel->SetDefault(PreviousValue);
+		Channel->GetInterface().AddKey(CurrentFrame, PreviousValue);
 
 		Track->AddSection(*Section);
 
@@ -93,7 +104,7 @@ UMovieSceneSection* FMovieScenePropertyRecorder<uint8>::AddSection(UObject* InOb
 template <>
 void FMovieScenePropertyRecorder<uint8>::AddKeyToSection(UMovieSceneSection* InSection, const FPropertyKey<uint8>& InKey)
 {
-	CastChecked<UMovieSceneByteSection>(InSection)->AddKey(InKey.Time, InKey.Value, EMovieSceneKeyInterpolation::Break);
+	InSection->GetChannelProxy().GetChannel<FMovieSceneByteChannel>(0)->GetInterface().AddKey(InKey.Time, InKey.Value);
 }
 
 template <>
@@ -117,10 +128,15 @@ UMovieSceneSection* FMovieScenePropertyRecorderEnum::AddSection(UObject* InObjec
 		}
 
 		UMovieSceneEnumSection* Section = Cast<UMovieSceneEnumSection>(Track->CreateNewSection());
-		Section->SetDefault(PreviousValue);
-		Section->SetStartTime(InTime);
-		Section->SetEndTime(InTime);
-		Section->AddKey(InTime, PreviousValue, EMovieSceneKeyInterpolation::Break);
+
+		FFrameRate   FrameResolution = Section->GetTypedOuter<UMovieScene>()->GetFrameResolution();
+		FFrameNumber CurrentFrame    = (InTime * FrameResolution).FloorToFrame();
+
+		Section->SetRange(TRange<FFrameNumber>::Inclusive(CurrentFrame, CurrentFrame));
+
+		FMovieSceneByteChannel* Channel = Section->GetChannelProxy().GetChannel<FMovieSceneByteChannel>(0);
+		Channel->SetDefault(PreviousValue);
+		Channel->GetInterface().AddKey(CurrentFrame, PreviousValue);
 
 		Track->AddSection(*Section);
 
@@ -132,7 +148,7 @@ UMovieSceneSection* FMovieScenePropertyRecorderEnum::AddSection(UObject* InObjec
 
 void FMovieScenePropertyRecorderEnum::AddKeyToSection(UMovieSceneSection* InSection, const FPropertyKey<int64>& InKey)
 {
-	CastChecked<UMovieSceneEnumSection>(InSection)->AddKey(InKey.Time, InKey.Value, EMovieSceneKeyInterpolation::Break);
+	InSection->GetChannelProxy().GetChannel<FMovieSceneByteChannel>(0)->GetInterface().AddKey(InKey.Time, InKey.Value);
 }
 
 void FMovieScenePropertyRecorderEnum::ReduceKeys(UMovieSceneSection* InSection)
@@ -157,10 +173,17 @@ UMovieSceneSection* FMovieScenePropertyRecorder<float>::AddSection(UObject* InOb
 		}
 
 		UMovieSceneFloatSection* Section = Cast<UMovieSceneFloatSection>(Track->CreateNewSection());
-		Section->SetDefault(PreviousValue);
-		Section->SetStartTime(InTime);
-		Section->SetEndTime(InTime);
-		Section->AddKey(InTime, PreviousValue, EMovieSceneKeyInterpolation::Break);
+
+		FFrameRate   FrameResolution = Section->GetTypedOuter<UMovieScene>()->GetFrameResolution();
+		FFrameNumber CurrentFrame    = (InTime * FrameResolution).FloorToFrame();
+
+		Section->SetRange(TRange<FFrameNumber>::Inclusive(CurrentFrame, CurrentFrame));
+
+
+		FMovieSceneFloatChannel* Channel = Section->GetChannelProxy().GetChannel<FMovieSceneFloatChannel>(0);
+		check(Channel);
+		Channel->SetDefault(PreviousValue);
+		Channel->AddCubicKey(CurrentFrame, PreviousValue, RCTM_Break);
 
 		Track->AddSection(*Section);
 
@@ -173,13 +196,14 @@ UMovieSceneSection* FMovieScenePropertyRecorder<float>::AddSection(UObject* InOb
 template <>
 void FMovieScenePropertyRecorder<float>::AddKeyToSection(UMovieSceneSection* InSection, const FPropertyKey<float>& InKey)
 {
-	CastChecked<UMovieSceneFloatSection>(InSection)->AddKey(InKey.Time, InKey.Value, EMovieSceneKeyInterpolation::Auto);
+	InSection->GetChannelProxy().GetChannel<FMovieSceneFloatChannel>(0)->AddCubicKey(InKey.Time, InKey.Value);
 }
 
 template <>
 void FMovieScenePropertyRecorder<float>::ReduceKeys(UMovieSceneSection* InSection)
 {
-	CastChecked<UMovieSceneFloatSection>(InSection)->GetFloatCurve().RemoveRedundantKeys(KINDA_SMALL_NUMBER);
+	FKeyDataOptimizationParams Params;
+	MovieScene::Optimize(InSection->GetChannelProxy().GetChannel<FMovieSceneFloatChannel>(0), Params);
 }
 
 template <>
@@ -200,18 +224,25 @@ UMovieSceneSection* FMovieScenePropertyRecorder<FColor>::AddSection(UObject* InO
 		}
 
 		UMovieSceneColorSection* Section = Cast<UMovieSceneColorSection>(Track->CreateNewSection());
-		Section->SetDefault(FColorKey(EKeyColorChannel::Red, PreviousValue.R, false));
-		Section->SetDefault(FColorKey(EKeyColorChannel::Green, PreviousValue.G, false));
-		Section->SetDefault(FColorKey(EKeyColorChannel::Blue, PreviousValue.B, false));
-		Section->SetDefault(FColorKey(EKeyColorChannel::Alpha, PreviousValue.A, false));
 
-		Section->SetStartTime(InTime);
-		Section->SetEndTime(InTime);
+		FFrameRate   FrameResolution = Section->GetTypedOuter<UMovieScene>()->GetFrameResolution();
+		FFrameNumber CurrentFrame    = (InTime * FrameResolution).FloorToFrame();
+
+		Section->SetRange(TRange<FFrameNumber>::Inclusive(CurrentFrame, CurrentFrame));
+
+		TArrayView<FMovieSceneFloatChannel*> FloatChannels = Section->GetChannelProxy().GetChannels<FMovieSceneFloatChannel>();
+
+		FloatChannels[0]->SetDefault(PreviousValue.R);
+		FloatChannels[0]->AddCubicKey(CurrentFrame, PreviousValue.R, RCTM_Break);
+
+		FloatChannels[1]->SetDefault(PreviousValue.G);
+		FloatChannels[1]->AddCubicKey(CurrentFrame, PreviousValue.G, RCTM_Break);
+
+		FloatChannels[2]->SetDefault(PreviousValue.B);
+		FloatChannels[2]->AddCubicKey(CurrentFrame, PreviousValue.B, RCTM_Break);
 
-		Section->AddKey(InTime, FColorKey(EKeyColorChannel::Red, PreviousValue.R, false), EMovieSceneKeyInterpolation::Break);
-		Section->AddKey(InTime, FColorKey(EKeyColorChannel::Green, PreviousValue.G, false), EMovieSceneKeyInterpolation::Break);
-		Section->AddKey(InTime, FColorKey(EKeyColorChannel::Blue, PreviousValue.B, false), EMovieSceneKeyInterpolation::Break);
-		Section->AddKey(InTime, FColorKey(EKeyColorChannel::Alpha, PreviousValue.A, false), EMovieSceneKeyInterpolation::Break);
+		FloatChannels[3]->SetDefault(PreviousValue.A);
+		FloatChannels[3]->AddCubicKey(CurrentFrame, PreviousValue.A, RCTM_Break);
 
 		Track->AddSection(*Section);
 
@@ -224,21 +255,23 @@ UMovieSceneSection* FMovieScenePropertyRecorder<FColor>::AddSection(UObject* InO
 template <>
 void FMovieScenePropertyRecorder<FColor>::AddKeyToSection(UMovieSceneSection* InSection, const FPropertyKey<FColor>& InKey)
 {
-	UMovieSceneColorSection* ColorSection = CastChecked<UMovieSceneColorSection>(InSection);
-	ColorSection->AddKey(InKey.Time, FColorKey(EKeyColorChannel::Red, InKey.Value.R, false), EMovieSceneKeyInterpolation::Auto);
-	ColorSection->AddKey(InKey.Time, FColorKey(EKeyColorChannel::Green, InKey.Value.G, false), EMovieSceneKeyInterpolation::Auto);
-	ColorSection->AddKey(InKey.Time, FColorKey(EKeyColorChannel::Blue, InKey.Value.B, false), EMovieSceneKeyInterpolation::Auto);
-	ColorSection->AddKey(InKey.Time, FColorKey(EKeyColorChannel::Alpha, InKey.Value.A, false), EMovieSceneKeyInterpolation::Auto);
+	TArrayView<FMovieSceneFloatChannel*> FloatChannels = InSection->GetChannelProxy().GetChannels<FMovieSceneFloatChannel>();
+	FloatChannels[0]->AddCubicKey(InKey.Time, InKey.Value.R);
+	FloatChannels[1]->AddCubicKey(InKey.Time, InKey.Value.G);
+	FloatChannels[2]->AddCubicKey(InKey.Time, InKey.Value.B);
+	FloatChannels[3]->AddCubicKey(InKey.Time, InKey.Value.A);
 }
 
 template <>
 void FMovieScenePropertyRecorder<FColor>::ReduceKeys(UMovieSceneSection* InSection)
 {
-	UMovieSceneColorSection* ColorSection = CastChecked<UMovieSceneColorSection>(InSection);
-	ColorSection->GetRedCurve().RemoveRedundantKeys(KINDA_SMALL_NUMBER);
-	ColorSection->GetGreenCurve().RemoveRedundantKeys(KINDA_SMALL_NUMBER);
-	ColorSection->GetBlueCurve().RemoveRedundantKeys(KINDA_SMALL_NUMBER);
-	ColorSection->GetAlphaCurve().RemoveRedundantKeys(KINDA_SMALL_NUMBER);
+	TArrayView<FMovieSceneFloatChannel*> FloatChannels = InSection->GetChannelProxy().GetChannels<FMovieSceneFloatChannel>();
+
+	FKeyDataOptimizationParams Params;
+	MovieScene::Optimize(FloatChannels[0], Params);
+	MovieScene::Optimize(FloatChannels[1], Params);
+	MovieScene::Optimize(FloatChannels[2], Params);
+	MovieScene::Optimize(FloatChannels[3], Params);
 }
 
 template <>
@@ -261,16 +294,21 @@ UMovieSceneSection* FMovieScenePropertyRecorder<FVector>::AddSection(UObject* In
 
 		UMovieSceneVectorSection* Section = Cast<UMovieSceneVectorSection>(Track->CreateNewSection());
 
-		Section->SetDefault(FVectorKey(EKeyVectorChannel::X, PreviousValue.X));
-		Section->SetDefault(FVectorKey(EKeyVectorChannel::Y, PreviousValue.Y));
-		Section->SetDefault(FVectorKey(EKeyVectorChannel::Z, PreviousValue.Z));
+		FFrameRate   FrameResolution = Section->GetTypedOuter<UMovieScene>()->GetFrameResolution();
+		FFrameNumber CurrentFrame    = (InTime * FrameResolution).FloorToFrame();
+
+		Section->SetRange(TRange<FFrameNumber>::Inclusive(CurrentFrame, CurrentFrame));
+
+		TArrayView<FMovieSceneFloatChannel*> FloatChannels = Section->GetChannelProxy().GetChannels<FMovieSceneFloatChannel>();
 
-		Section->SetStartTime(InTime);
-		Section->SetEndTime(InTime);
-		
-		Section->AddKey(InTime, FVectorKey(EKeyVectorChannel::X, PreviousValue.X), EMovieSceneKeyInterpolation::Break);
-		Section->AddKey(InTime, FVectorKey(EKeyVectorChannel::Y, PreviousValue.Y), EMovieSceneKeyInterpolation::Break);
-		Section->AddKey(InTime, FVectorKey(EKeyVectorChannel::Z, PreviousValue.Z), EMovieSceneKeyInterpolation::Break);
+		FloatChannels[0]->SetDefault(PreviousValue.X);
+		FloatChannels[0]->AddCubicKey(CurrentFrame, PreviousValue.X, RCTM_Break);
+
+		FloatChannels[1]->SetDefault(PreviousValue.Y);
+		FloatChannels[1]->AddCubicKey(CurrentFrame, PreviousValue.Y, RCTM_Break);
+
+		FloatChannels[2]->SetDefault(PreviousValue.Z);
+		FloatChannels[2]->AddCubicKey(CurrentFrame, PreviousValue.Z, RCTM_Break);
 
 		Track->AddSection(*Section);
 
@@ -283,21 +321,21 @@ UMovieSceneSection* FMovieScenePropertyRecorder<FVector>::AddSection(UObject* In
 template <>
 void FMovieScenePropertyRecorder<FVector>::AddKeyToSection(UMovieSceneSection* InSection, const FPropertyKey<FVector>& InKey)
 {
-	UMovieSceneVectorSection* VectorSection = CastChecked<UMovieSceneVectorSection>(InSection);
-
-	VectorSection->AddKey(InKey.Time, FVectorKey(EKeyVectorChannel::X, InKey.Value.X), EMovieSceneKeyInterpolation::Auto);
-	VectorSection->AddKey(InKey.Time, FVectorKey(EKeyVectorChannel::Y, InKey.Value.Y), EMovieSceneKeyInterpolation::Auto);
-	VectorSection->AddKey(InKey.Time, FVectorKey(EKeyVectorChannel::Z, InKey.Value.Z), EMovieSceneKeyInterpolation::Auto);
+	TArrayView<FMovieSceneFloatChannel*> FloatChannels = InSection->GetChannelProxy().GetChannels<FMovieSceneFloatChannel>();
+	FloatChannels[0]->AddCubicKey(InKey.Time, InKey.Value.X);
+	FloatChannels[1]->AddCubicKey(InKey.Time, InKey.Value.Y);
+	FloatChannels[2]->AddCubicKey(InKey.Time, InKey.Value.Z);
 }
 
 template <>
 void FMovieScenePropertyRecorder<FVector>::ReduceKeys(UMovieSceneSection* InSection)
 {
-	UMovieSceneVectorSection* VectorSection = CastChecked<UMovieSceneVectorSection>(InSection);
+	TArrayView<FMovieSceneFloatChannel*> FloatChannels = InSection->GetChannelProxy().GetChannels<FMovieSceneFloatChannel>();
 
-	VectorSection->GetCurve(0).RemoveRedundantKeys(KINDA_SMALL_NUMBER);
-	VectorSection->GetCurve(1).RemoveRedundantKeys(KINDA_SMALL_NUMBER);
-	VectorSection->GetCurve(2).RemoveRedundantKeys(KINDA_SMALL_NUMBER);
+	FKeyDataOptimizationParams Params;
+	MovieScene::Optimize(FloatChannels[0], Params);
+	MovieScene::Optimize(FloatChannels[1], Params);
+	MovieScene::Optimize(FloatChannels[2], Params);
 }
 
 template class FMovieScenePropertyRecorder<bool>;
diff --git a/Engine/Source/Editor/SequenceRecorder/Private/Sections/MovieScenePropertyRecorder.h b/Engine/Source/Editor/SequenceRecorder/Private/Sections/MovieScenePropertyRecorder.h
--- a/Engine/Source/Editor/SequenceRecorder/Private/Sections/MovieScenePropertyRecorder.h
+++ b/Engine/Source/Editor/SequenceRecorder/Private/Sections/MovieScenePropertyRecorder.h
@@ -6,6 +6,7 @@
 #include "MovieSceneCommonHelpers.h"
 #include "MovieSceneSection.h"
 #include "SequenceRecorderSettings.h"
+#include "MovieScene.h"
 
 
 /** Interface for a generic property recorder */
@@ -24,7 +25,7 @@ template <typename PropertyType>
 struct FPropertyKey
 {
 	PropertyType Value;
-	float Time;
+	FFrameNumber Time;
 };
 
 /** Recorder for a simple property of type PropertyType */
@@ -52,13 +53,16 @@ public:
 	{
 		if (InObjectToRecord != nullptr)
 		{
-			MovieSceneSection->SetEndTime(InCurrentTime);
+			FFrameRate   FrameResolution = MovieSceneSection->GetTypedOuter<UMovieScene>()->GetFrameResolution();
+			FFrameNumber CurrentFrame    = (InCurrentTime * FrameResolution).FloorToFrame();
+
+			MovieSceneSection->ExpandToFrame(CurrentFrame);
 
 			PropertyType NewValue = Binding.GetCurrentValue<PropertyType>(*InObjectToRecord);
 			if (ShouldAddNewKey(NewValue))
 			{
 				FPropertyKey<PropertyType> Key;
-				Key.Time = InCurrentTime;
+				Key.Time = CurrentFrame;
 				Key.Value = NewValue;
 
 				Keys.Add(Key);
@@ -136,13 +140,16 @@ public:
 	{
 		if (InObjectToRecord != nullptr)
 		{
-			MovieSceneSection->SetEndTime(InCurrentTime);
+			FFrameRate   FrameResolution = MovieSceneSection->GetTypedOuter<UMovieScene>()->GetFrameResolution();
+			FFrameNumber CurrentFrame    = (InCurrentTime * FrameResolution).FloorToFrame();
+
+			MovieSceneSection->ExpandToFrame(CurrentFrame);
 
 			int64 NewValue = Binding.GetCurrentValueForEnum(*InObjectToRecord);
 			if (ShouldAddNewKey(NewValue))
 			{
 				FPropertyKey<int64> Key;
-				Key.Time = InCurrentTime;
+				Key.Time = CurrentFrame;
 				Key.Value = NewValue;
 
 				Keys.Add(Key);
