Sequence Recorder: Fix StopRecording() condition when there is an actor recording that doesn't have a valid actor to record. Fix bool recorder not using bool channels.

diff --git a/Engine/Source/Editor/SequenceRecorder/Private/Sections/MovieScenePropertyRecorder.cpp b/Engine/Source/Editor/SequenceRecorder/Private/Sections/MovieScenePropertyRecorder.cpp
--- a/Engine/Source/Editor/SequenceRecorder/Private/Sections/MovieScenePropertyRecorder.cpp
+++ b/Engine/Source/Editor/SequenceRecorder/Private/Sections/MovieScenePropertyRecorder.cpp
@@ -58,9 +58,9 @@ UMovieSceneSection* FMovieScenePropertyRecorder<bool>::AddSection(UObject* InObj
 
 		Section->TimecodeSource = SequenceRecorderUtils::GetTimecodeSource();
 
-		TArrayView<FMovieSceneFloatChannel*> FloatChannels = Section->GetChannelProxy().GetChannels<FMovieSceneFloatChannel>();
-		FloatChannels[0]->SetDefault(PreviousValue);
-		FloatChannels[0]->AddCubicKey(CurrentFrame, PreviousValue, RCTM_Break);
+		FMovieSceneBoolChannel* BoolChannel = Section->GetChannelProxy().GetChannels<FMovieSceneBoolChannel>()[0];
+		BoolChannel->SetDefault(PreviousValue);
+		BoolChannel->GetData().AddKey(CurrentFrame, PreviousValue);
 
 		Track->AddSection(*Section);
 
diff --git a/Engine/Source/Editor/SequenceRecorder/Private/SequenceRecorder.cpp b/Engine/Source/Editor/SequenceRecorder/Private/SequenceRecorder.cpp
--- a/Engine/Source/Editor/SequenceRecorder/Private/SequenceRecorder.cpp
+++ b/Engine/Source/Editor/SequenceRecorder/Private/SequenceRecorder.cpp
@@ -367,7 +367,7 @@ void FSequenceRecorder::Tick(float DeltaSeconds)
 			bool bAllFinished = true;
 			for(UActorRecording* Recording : QueuedRecordings)
 			{
-				if(Recording->IsRecording())
+				if(Recording->GetActorToRecord() && Recording->IsRecording())
 				{
 					bAllFinished = false;
 					break;
