Sequence Recorder: Fix case where first recording is null but there are other valid recordings.

StopRecording wasn't getting called because the EndPIE delegate wasn't bound since ActorWorld was incorrectly set.

diff --git a/Engine/Source/Editor/SequenceRecorder/Private/SequenceRecorder.cpp b/Engine/Source/Editor/SequenceRecorder/Private/SequenceRecorder.cpp
--- a/Engine/Source/Editor/SequenceRecorder/Private/SequenceRecorder.cpp
+++ b/Engine/Source/Editor/SequenceRecorder/Private/SequenceRecorder.cpp
@@ -580,10 +580,20 @@ bool FSequenceRecorder::StartRecordingInternal(UWorld* World)
 
 	const USequenceRecorderSettings* Settings = GetDefault<USequenceRecorderSettings>();
 
-	UWorld* ActorWorld = nullptr;
-	if(World != nullptr || (QueuedRecordings.Num() > 0 && QueuedRecordings[0]->GetActorToRecord() != nullptr))
+	UWorld* ActorWorld = World;
+	if(ActorWorld == nullptr)
 	{
-		ActorWorld = World != nullptr ? World : QueuedRecordings[0]->GetActorToRecord()->GetWorld();
+		for (auto QueuedRecording : QueuedRecordings)
+		{
+			if (QueuedRecording->GetActorToRecord() != nullptr)
+			{
+				if (QueuedRecording->GetActorToRecord()->GetWorld() != nullptr)
+				{
+					ActorWorld = QueuedRecording->GetActorToRecord()->GetWorld();
+					break;
+				}
+			}
+		}
 	}
 
 	if(Settings->bRecordWorldSettingsActor)
