Sequence Recorder: Move ActorRecording.h to public.

diff --git a/Engine/Source/Editor/SequenceRecorder/Private/ActorRecording.cpp b/Engine/Source/Editor/SequenceRecorder/Private/ActorRecording.cpp
--- a/Engine/Source/Editor/SequenceRecorder/Private/ActorRecording.cpp
+++ b/Engine/Source/Editor/SequenceRecorder/Private/ActorRecording.cpp
@@ -10,6 +10,7 @@
 #include "LevelSequenceBindingReference.h"
 #include "SequenceRecorderSettings.h"
 #include "SequenceRecorderUtils.h"
+#include "Sections/MovieSceneAnimationSectionRecorder.h"
 #include "Sections/MovieScene3DTransformSectionRecorderSettings.h"
 #include "Sections/MovieSceneSubSection.h"
 #include "Tracks/MovieSceneSubTrack.h"
diff --git a/Engine/Source/Editor/SequenceRecorder/Private/ActorRecording.h b/Engine/Source/Editor/SequenceRecorder/Private/ActorRecording.h
deleted file mode 100644
diff --git a/Engine/Source/Editor/SequenceRecorder/Private/SequenceRecorderActorGroup.cpp b/Engine/Source/Editor/SequenceRecorder/Private/SequenceRecorderActorGroup.cpp
--- a/Engine/Source/Editor/SequenceRecorder/Private/SequenceRecorderActorGroup.cpp
+++ b/Engine/Source/Editor/SequenceRecorder/Private/SequenceRecorderActorGroup.cpp
@@ -1,5 +1,8 @@
+// Copyright 1998-2018 Epic Games, Inc. All Rights Reserved.
+
 #include "SequenceRecorderActorGroup.h"
 #include "SequenceRecorder.h"
+#include "ActorRecording.h"
 
 void USequenceRecorderActorGroup::PostEditChangeChainProperty(struct FPropertyChangedChainEvent& PropertyChangedEvent)
 {
diff --git a/Engine/Source/Editor/SequenceRecorder/Public/ActorRecording.h b/Engine/Source/Editor/SequenceRecorder/Public/ActorRecording.h
new file mode 100644
--- /dev/null
+++ b/Engine/Source/Editor/SequenceRecorder/Public/ActorRecording.h
@@ -0,0 +1,161 @@
+// Copyright 1998-2018 Epic Games, Inc. All Rights Reserved.
+
+#pragma once
+
+#include "CoreMinimal.h"
+#include "UObject/ObjectMacros.h"
+#include "UObject/Object.h"
+#include "Misc/Guid.h"
+#include "Animation/AnimationRecordingSettings.h"
+#include "IMovieSceneSectionRecorder.h"
+#include "ActorRecordingSettings.h"
+#include "UObject/ObjectKey.h"
+#include "GameFramework/Actor.h"
+
+#include "ActorRecording.generated.h"
+
+class UAnimSequence;
+class ULevelSequence;
+class UMovieScene;
+class USceneComponent;
+class FMovieSceneAnimationSectionRecorder;
+
+UCLASS(MinimalAPI)
+class UActorRecording : public UObject
+{
+	GENERATED_UCLASS_BODY()
+
+public:
+	/** Check whether it is worth recording this actor - i.e. is it going to affect the end result of the sequence */
+	static bool IsRelevantForRecording(AActor* Actor);
+
+	/** Start this queued recording. Sequence can be nullptr */
+	bool StartRecording(class ULevelSequence* CurrentSequence = nullptr, float CurrentSequenceTime = 0.0f, const FString& BaseAssetPath = FString(), const FString& SessionName = FString());
+
+	/** Stop this recording. Has no effect if we are not currently recording. Sequence can be nullptr */
+	bool StopRecording(class ULevelSequence* CurrentSequence = nullptr, float CurrentSequenceTime = 0.0f);
+
+	/** Tick this recording */
+	void Tick(float DeltaSeconds, ULevelSequence* CurrentSequence = nullptr, float CurrentSequenceTime = 0.0f);
+
+	/** Whether we are currently recording */
+	bool IsRecording() const;
+
+	/** Simulate a de-spawned actor */
+	void InvalidateObjectToRecord();
+
+	/** Get the Guid that identifies our spawnable in a recorded sequence */
+	const FGuid& GetSpawnableGuid() const
+	{
+		return Guid;
+	}
+
+	/** Get the actor to record. This finds the corresponding actor in the Simulation / PIE world. */
+	AActor* GetActorToRecord() const;
+
+	/** Set the actor to record */
+	void SetActorToRecord(AActor* InActor);
+
+	/** Get the active level sequence, optionally overridden by the target level sequence */
+	ULevelSequence* GetActiveLevelSequence(ULevelSequence* InLevelSequence);
+
+	/** Get target name */
+	FString GetTargetName(AActor* InActor);
+
+	/** Get the object binding for the actor if it exists in the level sequence either as a track with the track name or a tag with the actor label */
+	FGuid GetActorInSequence(AActor* InActor, ULevelSequence* CurrentSequence);
+
+private:
+
+	/** Get whether should duplicate the level sequence before recording into it */
+	bool ShouldDuplicateLevelSequence();
+
+	/** Check component validity for recording */
+	bool ValidComponent(USceneComponent* SceneComponent) const;
+
+	/** Adds us to a folder for better sequence organization */
+	void FindOrAddFolder(UMovieScene* MovieScene);
+
+	/** Start recording actor properties to a sequence */
+	void StartRecordingActorProperties(ULevelSequence* CurrentSequence, float CurrentSequenceTime);
+
+	/** Start recording component properties to a sequence */
+	TSharedPtr<class FMovieSceneAnimationSectionRecorder> StartRecordingComponentProperties(const FName& BindingName, USceneComponent* SceneComponent, UObject* BindingContext, ULevelSequence* CurrentSequence, float CurrentSequenceTime, const FAnimationRecordingSettings& InAnimationSettings, UAnimSequence* InTargetSequence);
+
+	/** Start recording components that are added at runtime */
+	void StartRecordingNewComponents(ULevelSequence* CurrentSequence, float CurrentSequenceTime);
+
+	/** Helper function to grab all scene components in the actor's hierarchy */
+	void GetSceneComponents(TArray<USceneComponent*>& OutArray, bool bIncludeNonCDO = true);
+
+	/** Sync up tracked components with the actor */
+	void SyncTrackedComponents(bool bIncludeNonCDO = true);
+
+	/** Ensure that we are recording any parents required for the specified component, and sort the specified array */
+	void ProcessNewComponentArray(TInlineComponentArray<USceneComponent*>& ProspectiveComponents) const;
+
+	/** UObject interface */
+	virtual void PostEditChangeProperty(struct FPropertyChangedEvent& PropertyChangedEvent) override;
+
+public:
+	UPROPERTY(EditAnywhere, Category = "Actor Recording")
+	FActorRecordingSettings ActorSettings;
+
+	/** Whether this actor is active and to be recorded when the 'Record' button is pressed. */
+	UPROPERTY(EditAnywhere, Category = "Actor Recording")
+	bool bActive;
+
+	/** Whether to create a level sequence for this actor recording */
+	UPROPERTY(EditAnywhere, Category = "Actor Recording")
+	bool bCreateLevelSequence;
+
+	/** The level sequence to record into */
+	UPROPERTY(EditAnywhere, Category = "Actor Recording", meta=(EditCondition = "bCreateLevelSequence"))
+	class ULevelSequence* TargetLevelSequence;
+
+	/** Optional target name to record to. If not specified, the actor label will be used */
+	UPROPERTY(EditAnywhere, Category = "Actor Recording", meta=(EditCondition = "bCreateLevelSequence"))
+	FText TargetName;
+
+	/** Specify the take number for the new recording */
+	UPROPERTY(EditAnywhere, Category = "Actor Recording", meta=(EditCondition = "bCreateLevelSequence"))
+	uint32 TakeNumber;
+
+	/** Whether we should specify the target animation or auto-create it */
+	UPROPERTY(EditAnywhere, Category = "Animation Recording")
+	bool bSpecifyTargetAnimation;
+
+	/** The target animation we want to record to */
+	UPROPERTY(EditAnywhere, Category = "Animation Recording", meta=(EditCondition = "bSpecifyTargetAnimation"))
+	class UAnimSequence* TargetAnimation;
+
+	/** The settings to apply to this actor's animation */
+	UPROPERTY(EditAnywhere, Category = "Animation Recording")
+	FAnimationRecordingSettings AnimationSettings;
+
+	/** Whether to record to 'possessable' (i.e. level-owned) or 'spawnable' (i.e. sequence-owned) actors. Defaults to the global setting. */
+	UPROPERTY(EditAnywhere, Category = "Actor Recording")
+	bool bRecordToPossessable;
+
+	/** Whether this actor recording was triggered from an actor spawn */
+	bool bWasSpawnedPostRecord;
+
+private:
+	/** The actor we want to record */
+	UPROPERTY(EditAnywhere, Category = "Actor Recording")
+	TLazyObjectPtr<AActor> ActorToRecord;
+
+	/** This actor's current set of section recorders */
+	TArray<TSharedPtr<class IMovieSceneSectionRecorder>> SectionRecorders;
+
+	/** Track components to check if any have changed */
+	TArray<TWeakObjectPtr<USceneComponent>> TrackedComponents;
+
+	TMap<FObjectKey, TWeakObjectPtr<USceneComponent>> DuplicatedDynamicComponents;
+
+	/** Flag to track whether we created new components */
+	bool bNewComponentAddedWhileRecording;
+
+	/** Guid that identifies our spawnable in a recorded sequence */
+	FGuid Guid;
+};
