Sequence Recorder: Change actor to record to soft object paths now that they get serialized through sequence recorder groups.

Added a few nullptr checks for actor groups and added transactions for adding/duplicating actors groups.
Fixed actor tags getting added multiple times.

diff --git a/Engine/Source/Editor/SequenceRecorder/Private/ActorRecording.cpp b/Engine/Source/Editor/SequenceRecorder/Private/ActorRecording.cpp
--- a/Engine/Source/Editor/SequenceRecorder/Private/ActorRecording.cpp
+++ b/Engine/Source/Editor/SequenceRecorder/Private/ActorRecording.cpp
@@ -480,7 +480,7 @@ void UActorRecording::StartRecordingActorProperties(ULevelSequence* CurrentSeque
 			{
 				if (!Possessable->Tags.Contains(*Actor->GetActorLabel()))
 				{
-					Possessable->Tags.Add(FName(*Actor->GetActorLabel()));
+					Possessable->Tags.AddUnique(FName(*Actor->GetActorLabel()));
 				}
 				Possessable->SetName(ObjectBindingName);
 			}
@@ -489,7 +489,7 @@ void UActorRecording::StartRecordingActorProperties(ULevelSequence* CurrentSeque
 			{
 				if (!Spawnable->Tags.Contains(*Actor->GetActorLabel()))
 				{
-					Spawnable->Tags.Add(FName(*Actor->GetActorLabel()));
+					Spawnable->Tags.AddUnique(FName(*Actor->GetActorLabel()));
 				}
 				Spawnable->SetName(ObjectBindingName);
 			}
@@ -801,9 +801,9 @@ bool UActorRecording::StopRecording(ULevelSequence* OriginalSequence, float Curr
 	}
 
 	// Swap to editor actor in case the actor was set while in PIE
-	if (ActorToRecord.IsValid())
+	if (AActor* Actor = Cast<AActor>(ActorToRecord.ResolveObject()))
 	{
-		if (AActor* EditorActor = EditorUtilities::GetEditorWorldCounterpartActor(ActorToRecord.Get()))
+		if (AActor* EditorActor = EditorUtilities::GetEditorWorldCounterpartActor(Actor))
 		{
 			ActorToRecord = EditorActor;
 		}
@@ -814,21 +814,21 @@ bool UActorRecording::StopRecording(ULevelSequence* OriginalSequence, float Curr
 
 bool UActorRecording::IsRecording() const
 {
-	return ActorToRecord.IsValid() && SectionRecorders.Num() > 0;
+	return GetActorToRecord() && SectionRecorders.Num() > 0;
 }
 
 AActor* UActorRecording::GetActorToRecord() const
 {
-	if (ActorToRecord.IsValid())
+	if (AActor* AssignedActor = Cast<AActor>(ActorToRecord.ResolveObject()))
 	{
-		AActor* OutActor = EditorUtilities::GetSimWorldCounterpartActor(ActorToRecord.Get());
+		AActor* OutActor = EditorUtilities::GetSimWorldCounterpartActor(AssignedActor);
 
 		if (OutActor != nullptr)
 		{
 			return OutActor;
 		}
 
-		return ActorToRecord.Get();
+		return AssignedActor;
 	}
 
 	return nullptr;
@@ -838,14 +838,14 @@ void UActorRecording::SetActorToRecord(AActor* InActor)
 {
 	ActorToRecord = InActor; 
 
-	if (ActorToRecord != nullptr)
+	if (InActor != nullptr)
 	{
 		bRecordToPossessable = false;
 
 		const USequenceRecorderSettings* Settings = GetDefault<USequenceRecorderSettings>();
 		for (const FSettingsForActorClass& SettingsForActorClass : Settings->PerActorSettings)
 		{
-			if (ActorToRecord->GetClass()->IsChildOf(SettingsForActorClass.Class))
+			if (InActor->GetClass()->IsChildOf(SettingsForActorClass.Class))
 			{
 				bRecordToPossessable = SettingsForActorClass.bRecordToPossessable;
 			}
@@ -857,14 +857,14 @@ void UActorRecording::PostEditChangeProperty(struct FPropertyChangedEvent& Prope
 {
 	if (PropertyChangedEvent.Property && PropertyChangedEvent.Property->GetFName() == GET_MEMBER_NAME_CHECKED(UActorRecording, ActorToRecord))
 	{
-		if (ActorToRecord != nullptr)
+		if (AActor* Actor = GetActorToRecord())
 		{
 			bRecordToPossessable = false;
 
 			const USequenceRecorderSettings* Settings = GetDefault<USequenceRecorderSettings>();
 			for (const FSettingsForActorClass& SettingsForActorClass : Settings->PerActorSettings)
 			{
-				if (ActorToRecord->GetClass()->IsChildOf(SettingsForActorClass.Class))
+				if (Actor->GetClass()->IsChildOf(SettingsForActorClass.Class))
 				{
 					bRecordToPossessable = SettingsForActorClass.bRecordToPossessable;
 				}
diff --git a/Engine/Source/Editor/SequenceRecorder/Private/SequenceRecorder.cpp b/Engine/Source/Editor/SequenceRecorder/Private/SequenceRecorder.cpp
--- a/Engine/Source/Editor/SequenceRecorder/Private/SequenceRecorder.cpp
+++ b/Engine/Source/Editor/SequenceRecorder/Private/SequenceRecorder.cpp
@@ -37,6 +37,7 @@
 #include "AssetToolsModule.h"
 #include "Camera/CameraActor.h"
 #include "Compilation/MovieSceneCompiler.h"
+#include "ScopedTransaction.h"
 
 #define LOCTEXT_NAMESPACE "SequenceRecorder"
 
@@ -1196,6 +1197,8 @@ TWeakObjectPtr<ASequenceRecorderGroup> FSequenceRecorder::GetRecordingGroupActor
 
 TWeakObjectPtr<USequenceRecorderActorGroup> FSequenceRecorder::AddRecordingGroup()
 {
+	const FScopedTransaction Transaction(LOCTEXT("AddRecordingGroup", "Add Actor Recording Group"));
+
 	TWeakObjectPtr<ASequenceRecorderGroup> GroupActor = GetRecordingGroupActor();
 	UWorld* EditorWorld = GEditor->GetEditorWorldContext().World();
 
@@ -1214,7 +1217,7 @@ TWeakObjectPtr<USequenceRecorderActorGroup> FSequenceRecorder::AddRecordingGroup
 
 	// Now add a new actor group to this actor
 	check(GroupActor.IsValid());
-	USequenceRecorderActorGroup* ActorGroup = NewObject<USequenceRecorderActorGroup>(GroupActor.Get());
+	USequenceRecorderActorGroup* ActorGroup = NewObject<USequenceRecorderActorGroup>(GroupActor.Get(), NAME_None, RF_Transactional);
 	if (!ExistingBasePath.Path.IsEmpty())
 	{
 		ActorGroup->SequenceRecordingBasePath = ExistingBasePath;
@@ -1246,6 +1249,8 @@ void FSequenceRecorder::RemoveCurrentRecordingGroup()
 		return;
 	}
 
+	const FScopedTransaction Transaction(LOCTEXT("RemoveActorRecordingGroup", "Remove Actor Recording Group"));
+
 	ClearQueuedRecordings();
 	TWeakObjectPtr<ASequenceRecorderGroup> GroupActor = GetRecordingGroupActor();
 	if (GroupActor.IsValid())
@@ -1265,14 +1270,24 @@ TWeakObjectPtr<USequenceRecorderActorGroup> FSequenceRecorder::DuplicateRecordin
 		BaseName = GetCurrentRecordingGroup().Get()->SequenceName;
 	}
 
+	const FScopedTransaction Transaction(LOCTEXT("DuplicateActorRecordingGroup", "Duplicate Actor Recording Group"));
+
 	USequenceRecorderActorGroup* DuplicatedGroup = DuplicateObject<USequenceRecorderActorGroup>(GetCurrentRecordingGroup().Get(), GetRecordingGroupActor().Get());
 	FString NewName = SequenceRecorderUtils::MakeNewGroupName(*DuplicatedGroup->SequenceRecordingBasePath.Path, BaseName, GetRecordingGroupNames());
 	DuplicatedGroup->GroupName = FName(*NewName);
 	DuplicatedGroup->SequenceName = NewName;
+	DuplicatedGroup->TargetLevelSequence = nullptr;
 	GetRecordingGroupActor().Get()->ActorGroups.Add(DuplicatedGroup);
 
 	// We'll invoke the standard load function so that it triggers everything to clear/update correctly.
-	return LoadRecordingGroup(DuplicatedGroup->GroupName);
+	TWeakObjectPtr<USequenceRecorderActorGroup> LoadedGroup = LoadRecordingGroup(DuplicatedGroup->GroupName);
+
+	if (OnRecordingGroupAddedDelegate.IsBound())
+	{
+		OnRecordingGroupAddedDelegate.Broadcast(LoadedGroup);
+	}
+
+	return LoadedGroup;
 }
 
 TArray<FName> FSequenceRecorder::GetRecordingGroupNames() const
@@ -1289,7 +1304,10 @@ TArray<FName> FSequenceRecorder::GetRecordingGroupNames() const
 			{
 				for (USequenceRecorderActorGroup* ActorGroup : GroupActor->ActorGroups)
 				{
-					GroupNames.Add(ActorGroup->GroupName);
+					if (ActorGroup)
+					{
+						GroupNames.Add(ActorGroup->GroupName);
+					}
 				}
 
 				// We only examine the first actor group in the map as it should contain all of our groups.
diff --git a/Engine/Source/Editor/SequenceRecorder/Public/ActorRecording.h b/Engine/Source/Editor/SequenceRecorder/Public/ActorRecording.h
--- a/Engine/Source/Editor/SequenceRecorder/Public/ActorRecording.h
+++ b/Engine/Source/Editor/SequenceRecorder/Public/ActorRecording.h
@@ -10,6 +10,7 @@
 #include "IMovieSceneSectionRecorder.h"
 #include "ActorRecordingSettings.h"
 #include "UObject/ObjectKey.h"
+#include "UObject/SoftObjectPath.h"
 #include "GameFramework/Actor.h"
 
 #include "ActorRecording.generated.h"
@@ -143,7 +144,7 @@ public:
 private:
 	/** The actor we want to record */
 	UPROPERTY(EditAnywhere, Category = "Actor Recording")
-	TLazyObjectPtr<AActor> ActorToRecord;
+	FSoftObjectPath ActorToRecord;
 
 	/** This actor's current set of section recorders */
 	TArray<TSharedPtr<class IMovieSceneSectionRecorder>> SectionRecorders;
diff --git a/Engine/Source/Editor/SequenceRecorder/Public/SequenceRecorderActorGroup.h b/Engine/Source/Editor/SequenceRecorder/Public/SequenceRecorderActorGroup.h
--- a/Engine/Source/Editor/SequenceRecorder/Public/SequenceRecorderActorGroup.h
+++ b/Engine/Source/Editor/SequenceRecorder/Public/SequenceRecorderActorGroup.h
@@ -78,7 +78,7 @@ public:
 	{
 		for (USequenceRecorderActorGroup* Group : ActorGroups)
 		{
-			if (Group->GroupName == Name)
+			if (Group && Group->GroupName == Name)
 			{
 				return Group;
 			}
