Sequence Recorder: Null World check in GetFirstPIEWorld.

diff --git a/Engine/Source/Editor/SequenceRecorder/Private/SequenceRecorder.cpp b/Engine/Source/Editor/SequenceRecorder/Private/SequenceRecorder.cpp
--- a/Engine/Source/Editor/SequenceRecorder/Private/SequenceRecorder.cpp
+++ b/Engine/Source/Editor/SequenceRecorder/Private/SequenceRecorder.cpp
@@ -163,7 +163,7 @@ static UWorld* GetFirstPIEWorld()
 {
 	for (const FWorldContext& Context : GEngine->GetWorldContexts())
 	{
-		if (Context.World()->IsPlayInEditor())
+		if (Context.World() != nullptr && Context.World()->IsPlayInEditor())
 		{
 			if(Context.World()->GetNetMode() == ENetMode::NM_Standalone ||
 				(Context.World()->GetNetMode() == ENetMode::NM_Client && Context.PIEInstance == 2))
