Sequence Recorder no longer shows all objects in the project and is instead limited to actors in the map that can be recorded.

diff --git a/Engine/Source/Editor/SequenceRecorder/Private/ActorRecording.cpp b/Engine/Source/Editor/SequenceRecorder/Private/ActorRecording.cpp
--- a/Engine/Source/Editor/SequenceRecorder/Private/ActorRecording.cpp
+++ b/Engine/Source/Editor/SequenceRecorder/Private/ActorRecording.cpp
@@ -801,7 +801,7 @@ bool UActorRecording::StopRecording(ULevelSequence* OriginalSequence, float Curr
 	}
 
 	// Swap to editor actor in case the actor was set while in PIE
-	if (AActor* Actor = Cast<AActor>(ActorToRecord.ResolveObject()))
+	if (AActor* Actor = ActorToRecord.Get())
 	{
 		if (AActor* EditorActor = EditorUtilities::GetEditorWorldCounterpartActor(Actor))
 		{
@@ -819,7 +819,7 @@ bool UActorRecording::IsRecording() const
 
 AActor* UActorRecording::GetActorToRecord() const
 {
-	if (AActor* AssignedActor = Cast<AActor>(ActorToRecord.ResolveObject()))
+	if (AActor* AssignedActor = ActorToRecord.Get())
 	{
 		AActor* OutActor = EditorUtilities::GetSimWorldCounterpartActor(AssignedActor);
 
diff --git a/Engine/Source/Editor/SequenceRecorder/Public/ActorRecording.h b/Engine/Source/Editor/SequenceRecorder/Public/ActorRecording.h
--- a/Engine/Source/Editor/SequenceRecorder/Public/ActorRecording.h
+++ b/Engine/Source/Editor/SequenceRecorder/Public/ActorRecording.h
@@ -144,7 +144,7 @@ public:
 private:
 	/** The actor we want to record */
 	UPROPERTY(EditAnywhere, Category = "Actor Recording")
-	FSoftObjectPath ActorToRecord;
+	TSoftObjectPtr<AActor> ActorToRecord;
 
 	/** This actor's current set of section recorders */
 	TArray<TSharedPtr<class IMovieSceneSectionRecorder>> SectionRecorders;
