Next Sequence Name now properly updates when duplicating a group. Next Sequence Name is also properly updated when there is an already conflicting asset name.

diff --git a/Engine/Source/Editor/SequenceRecorder/Private/SequenceRecorder.cpp b/Engine/Source/Editor/SequenceRecorder/Private/SequenceRecorder.cpp
--- a/Engine/Source/Editor/SequenceRecorder/Private/SequenceRecorder.cpp
+++ b/Engine/Source/Editor/SequenceRecorder/Private/SequenceRecorder.cpp
@@ -37,7 +37,6 @@
 #include "AssetToolsModule.h"
 #include "Camera/CameraActor.h"
 #include "Compilation/MovieSceneCompiler.h"
-#include "ScopedTransaction.h"
 
 #define LOCTEXT_NAMESPACE "SequenceRecorder"
 
@@ -163,7 +162,7 @@ static UWorld* GetFirstPIEWorld()
 {
 	for (const FWorldContext& Context : GEngine->GetWorldContexts())
 	{
-		if (Context.World() != nullptr && Context.World()->IsPlayInEditor())
+		if (Context.World()->IsPlayInEditor())
 		{
 			if(Context.World()->GetNetMode() == ENetMode::NM_Standalone ||
 				(Context.World()->GetNetMode() == ENetMode::NM_Client && Context.PIEInstance == 2))
@@ -1152,16 +1151,20 @@ void FSequenceRecorder::RefreshNextSequence()
 		SequenceName = GetSequenceRecordingName().Len() > 0 ? GetSequenceRecordingName() : TEXT("RecordedSequence");
 	}
 
-	// Cache the name of the next sequence we will try to record to
-	NextSequenceName = SequenceRecorderUtils::MakeNewAssetName(GetSequenceRecordingBasePath(), SequenceName);
+	// Cache the name of the next sequence we will try to record to. Assets are recorded into a folder with their desired name, so we need to append that
+	// to the base path before checking for unique names.
+	FString AssetPath = FString::Printf(TEXT("%s/%s"), *GetSequenceRecordingBasePath(), *GetSequenceRecordingName());
+	NextSequenceName = SequenceRecorderUtils::MakeNewAssetName(AssetPath, SequenceName);
 }
 
 void FSequenceRecorder::ForceRefreshNextSequence()
 {
 	SequenceName = GetSequenceRecordingName().Len() > 0 ? GetSequenceRecordingName() : TEXT("RecordedSequence");
 
-	// Cache the name of the next sequence we will try to record to
-	NextSequenceName = SequenceRecorderUtils::MakeNewAssetName(GetSequenceRecordingBasePath(), SequenceName);
+	// Cache the name of the next sequence we will try to record to. Assets are recorded into a folder with their desired name, so we need to append that
+	// to the base path before checking for unique names.
+	FString AssetPath = FString::Printf(TEXT("%s/%s"), *GetSequenceRecordingBasePath(), *GetSequenceRecordingName());
+	NextSequenceName = SequenceRecorderUtils::MakeNewAssetName(AssetPath, SequenceName);
 }
 
 TWeakObjectPtr<ASequenceRecorderGroup> FSequenceRecorder::GetRecordingGroupActor()
@@ -1197,8 +1200,6 @@ TWeakObjectPtr<ASequenceRecorderGroup> FSequenceRecorder::GetRecordingGroupActor
 
 TWeakObjectPtr<USequenceRecorderActorGroup> FSequenceRecorder::AddRecordingGroup()
 {
-	const FScopedTransaction Transaction(LOCTEXT("AddRecordingGroup", "Add Actor Recording Group"));
-
 	TWeakObjectPtr<ASequenceRecorderGroup> GroupActor = GetRecordingGroupActor();
 	UWorld* EditorWorld = GEditor->GetEditorWorldContext().World();
 
@@ -1217,7 +1218,7 @@ TWeakObjectPtr<USequenceRecorderActorGroup> FSequenceRecorder::AddRecordingGroup
 
 	// Now add a new actor group to this actor
 	check(GroupActor.IsValid());
-	USequenceRecorderActorGroup* ActorGroup = NewObject<USequenceRecorderActorGroup>(GroupActor.Get(), NAME_None, RF_Transactional);
+	USequenceRecorderActorGroup* ActorGroup = NewObject<USequenceRecorderActorGroup>(GroupActor.Get());
 	if (!ExistingBasePath.Path.IsEmpty())
 	{
 		ActorGroup->SequenceRecordingBasePath = ExistingBasePath;
@@ -1233,6 +1234,8 @@ TWeakObjectPtr<USequenceRecorderActorGroup> FSequenceRecorder::AddRecordingGroup
 
 	// And then select our new object by default
 	CurrentRecorderGroup = ActorGroup;
+
+	ForceRefreshNextSequence();
 	
 	if (OnRecordingGroupAddedDelegate.IsBound())
 	{
@@ -1249,8 +1252,6 @@ void FSequenceRecorder::RemoveCurrentRecordingGroup()
 		return;
 	}
 
-	const FScopedTransaction Transaction(LOCTEXT("RemoveActorRecordingGroup", "Remove Actor Recording Group"));
-
 	ClearQueuedRecordings();
 	TWeakObjectPtr<ASequenceRecorderGroup> GroupActor = GetRecordingGroupActor();
 	if (GroupActor.IsValid())
@@ -1270,24 +1271,14 @@ TWeakObjectPtr<USequenceRecorderActorGroup> FSequenceRecorder::DuplicateRecordin
 		BaseName = GetCurrentRecordingGroup().Get()->SequenceName;
 	}
 
-	const FScopedTransaction Transaction(LOCTEXT("DuplicateActorRecordingGroup", "Duplicate Actor Recording Group"));
-
 	USequenceRecorderActorGroup* DuplicatedGroup = DuplicateObject<USequenceRecorderActorGroup>(GetCurrentRecordingGroup().Get(), GetRecordingGroupActor().Get());
 	FString NewName = SequenceRecorderUtils::MakeNewGroupName(*DuplicatedGroup->SequenceRecordingBasePath.Path, BaseName, GetRecordingGroupNames());
 	DuplicatedGroup->GroupName = FName(*NewName);
 	DuplicatedGroup->SequenceName = NewName;
-	DuplicatedGroup->TargetLevelSequence = nullptr;
 	GetRecordingGroupActor().Get()->ActorGroups.Add(DuplicatedGroup);
 
 	// We'll invoke the standard load function so that it triggers everything to clear/update correctly.
-	TWeakObjectPtr<USequenceRecorderActorGroup> LoadedGroup = LoadRecordingGroup(DuplicatedGroup->GroupName);
-
-	if (OnRecordingGroupAddedDelegate.IsBound())
-	{
-		OnRecordingGroupAddedDelegate.Broadcast(LoadedGroup);
-	}
-
-	return LoadedGroup;
+	return LoadRecordingGroup(DuplicatedGroup->GroupName);
 }
 
 TArray<FName> FSequenceRecorder::GetRecordingGroupNames() const
@@ -1304,10 +1295,7 @@ TArray<FName> FSequenceRecorder::GetRecordingGroupNames() const
 			{
 				for (USequenceRecorderActorGroup* ActorGroup : GroupActor->ActorGroups)
 				{
-					if (ActorGroup)
-					{
-						GroupNames.Add(ActorGroup->GroupName);
-					}
+					GroupNames.Add(ActorGroup->GroupName);
 				}
 
 				// We only examine the first actor group in the map as it should contain all of our groups.
@@ -1355,16 +1343,18 @@ TWeakObjectPtr<USequenceRecorderActorGroup> FSequenceRecorder::LoadRecordingGrou
 					QueuedRecordings.Add(ActorRecording);
 				}
 			}
-			RefreshNextSequence();
+			ForceRefreshNextSequence();
 			return CurrentRecorderGroup;
 		}
 	}
 
 	// We either don't have a group actor or we can't find a group by that name, clear anything we have loaded.
 	// This lets the UI handle switching back to profile "None".
-	RefreshNextSequence();
 	ClearQueuedRecordings();
 	CurrentRecorderGroup = nullptr;
+
+	// Refresh the next sequence after nulling out the recording group so we get the default name.
+	ForceRefreshNextSequence();
 	return nullptr;
 }
 
