Sequence Recorder Recording Groups no longer automatically post-fix a letter onto the name when renamed by the user.

diff --git a/Engine/Source/Editor/SequenceRecorder/Private/ActorGroupDetailsCustomization.cpp b/Engine/Source/Editor/SequenceRecorder/Private/ActorGroupDetailsCustomization.cpp
--- a/Engine/Source/Editor/SequenceRecorder/Private/ActorGroupDetailsCustomization.cpp
+++ b/Engine/Source/Editor/SequenceRecorder/Private/ActorGroupDetailsCustomization.cpp
@@ -209,14 +209,30 @@ void FActorGroupDetailsCustomization::HandleRecordingGroupNameCommitted(const FT
 		TWeakObjectPtr<USequenceRecorderActorGroup> CurrentGroup = FSequenceRecorder::Get().GetCurrentRecordingGroup();
 		if (CurrentGroup.IsValid())
 		{
-			TArray<FName> ExistingGroupNames = FSequenceRecorder::Get().GetRecordingGroupNames();
-			ExistingGroupNames.Remove(FName(*InText.ToString()));
-			FString NewName = SequenceRecorderUtils::MakeNewGroupName(*CurrentGroup.Get()->SequenceRecordingBasePath.Path, InText.ToString(), ExistingGroupNames);
-			CurrentGroup->GroupName = FName(*NewName);
-			CurrentGroup->SequenceName = NewName;
+			const TArray<FName> ExistingGroupNames = FSequenceRecorder::Get().GetRecordingGroupNames();
+			const FString NewNameAsString = InText.ToString();
+			FName NewName = FName(*NewNameAsString);
+
+			// If they're trying to rename to the same name, just early out. If we try to go through with the rename,
+			// it'll see that the name is already in the list and postfix it with "_A" when renaming "Foo" to "Foo".
+			if (CurrentGroup->GroupName == NewName)
+			{
+				return;
+			}
+
+			// If they're trying to rename it to something already on our list (another entry) then we forcibly change
+			// their name to something unique.
+			if (ExistingGroupNames.Contains(NewName))
+			{
+				NewName = FName(*SequenceRecorderUtils::MakeNewGroupName(*CurrentGroup.Get()->SequenceRecordingBasePath.Path, InText.ToString(), ExistingGroupNames));
+			}
+
+			// Re-assign the name of the recording group and update the sequence.
+			CurrentGroup->GroupName = NewName;
+			CurrentGroup->SequenceName = NewNameAsString;
 
 			// cbb: Force load to update sequence names, etc
-			SequenceRecorder.Pin()->HandleLoadRecordingActorGroup(*NewName);
+			SequenceRecorder.Pin()->HandleLoadRecordingActorGroup(NewName);
 		}
 	}
 }
