Sequence Recorder: Reset TargetLevelSequence and takes when duplicating or changing actor group names.

diff --git a/Engine/Source/Editor/SequenceRecorder/Private/ActorGroupDetailsCustomization.cpp b/Engine/Source/Editor/SequenceRecorder/Private/ActorGroupDetailsCustomization.cpp
--- a/Engine/Source/Editor/SequenceRecorder/Private/ActorGroupDetailsCustomization.cpp
+++ b/Engine/Source/Editor/SequenceRecorder/Private/ActorGroupDetailsCustomization.cpp
@@ -230,6 +230,15 @@ void FActorGroupDetailsCustomization::HandleRecordingGroupNameCommitted(const FT
 			// Re-assign the name of the recording group and update the sequence.
 			CurrentGroup->GroupName = NewName;
 			CurrentGroup->SequenceName = NewNameAsString;
+			CurrentGroup->TargetLevelSequence = nullptr;
+
+			for (UActorRecording* ActorRecording : CurrentGroup->RecordedActors)
+			{
+				if (ActorRecording != nullptr)
+				{
+					ActorRecording->TakeNumber = 1;
+				}
+			}
 
 			// cbb: Force load to update sequence names, etc
 			SequenceRecorder.Pin()->HandleLoadRecordingActorGroup(NewName);
diff --git a/Engine/Source/Editor/SequenceRecorder/Private/SequenceRecorder.cpp b/Engine/Source/Editor/SequenceRecorder/Private/SequenceRecorder.cpp
--- a/Engine/Source/Editor/SequenceRecorder/Private/SequenceRecorder.cpp
+++ b/Engine/Source/Editor/SequenceRecorder/Private/SequenceRecorder.cpp
@@ -37,6 +37,7 @@
 #include "AssetToolsModule.h"
 #include "Camera/CameraActor.h"
 #include "Compilation/MovieSceneCompiler.h"
+#include "ScopedTransaction.h"
 
 #define LOCTEXT_NAMESPACE "SequenceRecorder"
 
@@ -1200,6 +1201,8 @@ TWeakObjectPtr<ASequenceRecorderGroup> FSequenceRecorder::GetRecordingGroupActor
 
 TWeakObjectPtr<USequenceRecorderActorGroup> FSequenceRecorder::AddRecordingGroup()
 {
+	const FScopedTransaction Transaction(LOCTEXT("AddRecordingGroup", "Add Actor Recording Group"));
+
 	TWeakObjectPtr<ASequenceRecorderGroup> GroupActor = GetRecordingGroupActor();
 	UWorld* EditorWorld = GEditor->GetEditorWorldContext().World();
 
@@ -1218,7 +1221,7 @@ TWeakObjectPtr<USequenceRecorderActorGroup> FSequenceRecorder::AddRecordingGroup
 
 	// Now add a new actor group to this actor
 	check(GroupActor.IsValid());
-	USequenceRecorderActorGroup* ActorGroup = NewObject<USequenceRecorderActorGroup>(GroupActor.Get());
+	USequenceRecorderActorGroup* ActorGroup = NewObject<USequenceRecorderActorGroup>(GroupActor.Get(), NAME_None, RF_Transactional);
 	if (!ExistingBasePath.Path.IsEmpty())
 	{
 		ActorGroup->SequenceRecordingBasePath = ExistingBasePath;
@@ -1252,6 +1255,8 @@ void FSequenceRecorder::RemoveCurrentRecordingGroup()
 		return;
 	}
 
+	const FScopedTransaction Transaction(LOCTEXT("RemoveActorRecordingGroup", "Remove Actor Recording Group"));
+
 	ClearQueuedRecordings();
 	TWeakObjectPtr<ASequenceRecorderGroup> GroupActor = GetRecordingGroupActor();
 	if (GroupActor.IsValid())
@@ -1271,14 +1276,33 @@ TWeakObjectPtr<USequenceRecorderActorGroup> FSequenceRecorder::DuplicateRecordin
 		BaseName = GetCurrentRecordingGroup().Get()->SequenceName;
 	}
 
+	const FScopedTransaction Transaction(LOCTEXT("DuplicateActorRecordingGroup", "Duplicate Actor Recording Group"));
+
 	USequenceRecorderActorGroup* DuplicatedGroup = DuplicateObject<USequenceRecorderActorGroup>(GetCurrentRecordingGroup().Get(), GetRecordingGroupActor().Get());
 	FString NewName = SequenceRecorderUtils::MakeNewGroupName(*DuplicatedGroup->SequenceRecordingBasePath.Path, BaseName, GetRecordingGroupNames());
 	DuplicatedGroup->GroupName = FName(*NewName);
 	DuplicatedGroup->SequenceName = NewName;
+	DuplicatedGroup->TargetLevelSequence = nullptr;
+
+	for (UActorRecording* ActorRecording : DuplicatedGroup->RecordedActors)
+	{
+		if (ActorRecording != nullptr)
+		{
+			ActorRecording->TakeNumber = 1;
+		}
+	}
+
 	GetRecordingGroupActor().Get()->ActorGroups.Add(DuplicatedGroup);
 
 	// We'll invoke the standard load function so that it triggers everything to clear/update correctly.
-	return LoadRecordingGroup(DuplicatedGroup->GroupName);
+	TWeakObjectPtr<USequenceRecorderActorGroup> LoadedGroup = LoadRecordingGroup(DuplicatedGroup->GroupName);
+
+	if (OnRecordingGroupAddedDelegate.IsBound())
+	{
+		OnRecordingGroupAddedDelegate.Broadcast(LoadedGroup);
+	}
+
+	return LoadedGroup;
 }
 
 TArray<FName> FSequenceRecorder::GetRecordingGroupNames() const
@@ -1295,7 +1319,10 @@ TArray<FName> FSequenceRecorder::GetRecordingGroupNames() const
 			{
 				for (USequenceRecorderActorGroup* ActorGroup : GroupActor->ActorGroups)
 				{
-					GroupNames.Add(ActorGroup->GroupName);
+					if (ActorGroup)
+					{
+						GroupNames.Add(ActorGroup->GroupName);
+					}
 				}
 
 				// We only examine the first actor group in the map as it should contain all of our groups.
diff --git a/Engine/Source/Editor/SequenceRecorder/Private/SequenceRecorderActorGroup.cpp b/Engine/Source/Editor/SequenceRecorder/Private/SequenceRecorderActorGroup.cpp
--- a/Engine/Source/Editor/SequenceRecorder/Private/SequenceRecorderActorGroup.cpp
+++ b/Engine/Source/Editor/SequenceRecorder/Private/SequenceRecorderActorGroup.cpp
@@ -10,6 +10,20 @@ void USequenceRecorderActorGroup::PostEditChangeChainProperty(struct FPropertyCh
 
 	if (PropertyChangedEvent.Property)
 	{
+		if (PropertyChangedEvent.Property->GetFName() == GET_MEMBER_NAME_CHECKED(USequenceRecorderActorGroup, SequenceName))
+		{
+			GroupName = *SequenceName;
+			TargetLevelSequence = nullptr;
+
+			for (UActorRecording* ActorRecording : RecordedActors)
+			{
+				if (ActorRecording != nullptr)
+				{
+					ActorRecording->TakeNumber = 1;
+				}
+			}
+		}
+
 		if (PropertyChangedEvent.Property->GetFName() == GET_MEMBER_NAME_CHECKED(USequenceRecorderActorGroup, SequenceName) ||
 			PropertyChangedEvent.Property->GetFName() == GET_MEMBER_NAME_CHECKED(USequenceRecorderActorGroup, SequenceRecordingBasePath))
 		{
