Fix bad merge
#rnx

diff --git a/Engine/Source/Editor/SequenceRecorder/Private/ActorRecording.cpp b/Engine/Source/Editor/SequenceRecorder/Private/ActorRecording.cpp
--- a/Engine/Source/Editor/SequenceRecorder/Private/ActorRecording.cpp
+++ b/Engine/Source/Editor/SequenceRecorder/Private/ActorRecording.cpp
@@ -1065,15 +1065,21 @@ void UActorRecording::StartRecordingNewComponents(ULevelSequence* CurrentSequenc
 								AttachToComponent = Cast<USceneComponent>(LocalActorComponent);
 							}
 
-							for (UActorComponent* Component : ObjectTemplate->GetComponents())
+							// If we don't have an attachment parent duplicated already, perform a name lookup
+							if (!AttachToComponent)
 							{
-								if (USceneComponent* Child = Cast<USceneComponent>(Component))
+								FName AttachName = SceneComponent->GetAttachParent()->GetFName();
+
+								for (UActorComponent* Component : ObjectTemplate->GetComponents())
 								{
-									CA_SUPPRESS(28182); // Dereferencing NULL pointer. 'Child' contains the same NULL value as 'AttachToComponent' did.
-									if (Child->GetFName() == AttachName)
+									if (USceneComponent* Child = Cast<USceneComponent>(Component))
 									{
-										AttachToComponent = Child;
-										break;
+										CA_SUPPRESS(28182); // Dereferencing NULL pointer. 'Child' contains the same NULL value as 'AttachToComponent' did.
+										if (Child->GetFName() == AttachName)
+										{
+											AttachToComponent = Child;
+											break;
+										}
 									}
 								}
 							}
