Sequence Recorder: Fix StopRecording() condition when there is an actor recording thatt doesn't have a valid actor to record.

diff --git a/Engine/Source/Editor/SequenceRecorder/Private/SequenceRecorder.cpp b/Engine/Source/Editor/SequenceRecorder/Private/SequenceRecorder.cpp
--- a/Engine/Source/Editor/SequenceRecorder/Private/SequenceRecorder.cpp
+++ b/Engine/Source/Editor/SequenceRecorder/Private/SequenceRecorder.cpp
@@ -357,7 +357,7 @@ void FSequenceRecorder::Tick(float DeltaSeconds)
 			bool bAllFinished = true;
 			for(UActorRecording* Recording : QueuedRecordings)
 			{
-				if(Recording->IsRecording())
+				if(Recording->GetActorToRecord() && Recording->IsRecording())
 				{
 					bAllFinished = false;
 					break;
