Sequence Recorder: If removing root animation from the animation asset, push the root animation to the transform recorder, which records to the transform track only when the animation is not in world space.

diff --git a/Engine/Source/Editor/SequenceRecorder/Private/ActorRecording.cpp b/Engine/Source/Editor/SequenceRecorder/Private/ActorRecording.cpp
--- a/Engine/Source/Editor/SequenceRecorder/Private/ActorRecording.cpp
+++ b/Engine/Source/Editor/SequenceRecorder/Private/ActorRecording.cpp
@@ -154,6 +154,14 @@ bool UActorRecording::StartRecording(ULevelSequence* CurrentSequence, float Curr
 				}
 			}
 
+			// If removing root animation from the animation asset, push the root animation 
+			// to the transform recorder and the transform track, which gets recorded only 
+			// when the animation is not in world space.
+			if (AnimationSettings.bRemoveRootAnimation)
+			{
+				AnimationSettings.bRecordInWorldSpace = false;
+			}
+
 			TSharedPtr<FMovieSceneAnimationSectionRecorder> AnimationRecorder = MakeShareable(new FMovieSceneAnimationSectionRecorder(AnimationSettings, TargetAnimation, AnimAssetPath, AnimAssetName));
 			AnimationRecorder->CreateSection(GetActorToRecord(), nullptr, FGuid(), 0.0f);
 			AnimationRecorder->Record(0.0f);
