Sequence Recorder: When recording animation in world space, the transform keys are compensated by the root bone. There was a bug when not recording animation in world space where there would be no transform keys. This is because the root bone values are all identity (since they're not in world space). To fix this, the transform recorder should record as usual, filling the buffered transforms from the actor, rather than swapping them in during finalize from the root bone.

diff --git a/Engine/Source/Editor/SequenceRecorder/Private/Sections/MovieScene3DTransformSectionRecorder.cpp b/Engine/Source/Editor/SequenceRecorder/Private/Sections/MovieScene3DTransformSectionRecorder.cpp
--- a/Engine/Source/Editor/SequenceRecorder/Private/Sections/MovieScene3DTransformSectionRecorder.cpp
+++ b/Engine/Source/Editor/SequenceRecorder/Private/Sections/MovieScene3DTransformSectionRecorder.cpp
@@ -114,8 +114,10 @@ void FMovieScene3DTransformSectionRecorder::FinalizeSection(float CurrentTime)
 	check (	BufferedTransforms.Times.Num() == BufferedTransforms.ScaleZ.Num());
 
 	// if we have a valid animation recorder, we need to build our transforms from the animation
-	// so we properly synchronize our keyframes
-	if(AnimRecorder.IsValid() && bWasRecording)
+	// so we properly synchronize our keyframes. This should only be done when recording animation 
+	// to the animation asset in world space because otherwise if recording to local space, the root 
+	// bone would resolve to identity and local space transform keys would be recorded.
+	if(AnimRecorder.IsValid() && bWasRecording && AnimRecorder->AnimationSettings.bRecordInWorldSpace)
 	{
 		check(BufferedTransforms.Times.Num() == 0);
 
@@ -355,7 +357,7 @@ void FMovieScene3DTransformSectionRecorder::Record(float CurrentTime)
 		if(bRecording)
 		{
 			// don't record from the transform of the component/actor if we are synchronizing with an animation
-			if (!AnimRecorder.IsValid())
+			if (!AnimRecorder.IsValid() || !AnimRecorder->AnimationSettings.bRecordInWorldSpace)
 			{
 				FTransform TransformToRecord;
 				if (GetTransformToRecord(TransformToRecord))
