Fix CIS

diff --git a/Engine/Source/Editor/SequenceRecorder/Private/Sections/MovieSceneAnimationSectionRecorder.cpp b/Engine/Source/Editor/SequenceRecorder/Private/Sections/MovieSceneAnimationSectionRecorder.cpp
--- a/Engine/Source/Editor/SequenceRecorder/Private/Sections/MovieSceneAnimationSectionRecorder.cpp
+++ b/Engine/Source/Editor/SequenceRecorder/Private/Sections/MovieSceneAnimationSectionRecorder.cpp
@@ -35,9 +35,9 @@ bool FMovieSceneAnimationSectionRecorderFactory::CanRecordObject(UObject* InObje
 }
 
 FMovieSceneAnimationSectionRecorder::FMovieSceneAnimationSectionRecorder(const FAnimationRecordingSettings& InAnimationSettings, UAnimSequence* InSpecifiedSequence, const FString& InAnimAssetPath, const FString& InAnimAssetName)
-	: AnimSequence(InSpecifiedSequence)
+	: AnimationSettings(InAnimationSettings)
+	, AnimSequence(InSpecifiedSequence)
 	, bRemoveRootTransform(true)
-	, AnimationSettings(InAnimationSettings)
 	, AnimAssetPath(InAnimAssetPath)
 	, AnimAssetName(InAnimAssetName)
 {
