Next Sequence Name now properly updates when duplicating a group. Next Sequence Name is also properly updated when there is an already conflicting asset name.

diff --git a/Engine/Source/Editor/SequenceRecorder/Private/SequenceRecorder.cpp b/Engine/Source/Editor/SequenceRecorder/Private/SequenceRecorder.cpp
--- a/Engine/Source/Editor/SequenceRecorder/Private/SequenceRecorder.cpp
+++ b/Engine/Source/Editor/SequenceRecorder/Private/SequenceRecorder.cpp
@@ -1151,16 +1151,20 @@ void FSequenceRecorder::RefreshNextSequence()
 		SequenceName = GetSequenceRecordingName().Len() > 0 ? GetSequenceRecordingName() : TEXT("RecordedSequence");
 	}
 
-	// Cache the name of the next sequence we will try to record to
-	NextSequenceName = SequenceRecorderUtils::MakeNewAssetName(GetSequenceRecordingBasePath(), SequenceName);
+	// Cache the name of the next sequence we will try to record to. Assets are recorded into a folder with their desired name, so we need to append that
+	// to the base path before checking for unique names.
+	FString AssetPath = FString::Printf(TEXT("%s/%s"), *GetSequenceRecordingBasePath(), *GetSequenceRecordingName());
+	NextSequenceName = SequenceRecorderUtils::MakeNewAssetName(AssetPath, SequenceName);
 }
 
 void FSequenceRecorder::ForceRefreshNextSequence()
 {
 	SequenceName = GetSequenceRecordingName().Len() > 0 ? GetSequenceRecordingName() : TEXT("RecordedSequence");
 
-	// Cache the name of the next sequence we will try to record to
-	NextSequenceName = SequenceRecorderUtils::MakeNewAssetName(GetSequenceRecordingBasePath(), SequenceName);
+	// Cache the name of the next sequence we will try to record to. Assets are recorded into a folder with their desired name, so we need to append that
+	// to the base path before checking for unique names.
+	FString AssetPath = FString::Printf(TEXT("%s/%s"), *GetSequenceRecordingBasePath(), *GetSequenceRecordingName());
+	NextSequenceName = SequenceRecorderUtils::MakeNewAssetName(AssetPath, SequenceName);
 }
 
 TWeakObjectPtr<ASequenceRecorderGroup> FSequenceRecorder::GetRecordingGroupActor()
@@ -1230,6 +1234,8 @@ TWeakObjectPtr<USequenceRecorderActorGroup> FSequenceRecorder::AddRecordingGroup
 
 	// And then select our new object by default
 	CurrentRecorderGroup = ActorGroup;
+
+	ForceRefreshNextSequence();
 	
 	if (OnRecordingGroupAddedDelegate.IsBound())
 	{
@@ -1337,16 +1343,18 @@ TWeakObjectPtr<USequenceRecorderActorGroup> FSequenceRecorder::LoadRecordingGrou
 					QueuedRecordings.Add(ActorRecording);
 				}
 			}
-			RefreshNextSequence();
+			ForceRefreshNextSequence();
 			return CurrentRecorderGroup;
 		}
 	}
 
 	// We either don't have a group actor or we can't find a group by that name, clear anything we have loaded.
 	// This lets the UI handle switching back to profile "None".
-	RefreshNextSequence();
 	ClearQueuedRecordings();
 	CurrentRecorderGroup = nullptr;
+
+	// Refresh the next sequence after nulling out the recording group so we get the default name.
+	ForceRefreshNextSequence();
 	return nullptr;
 }
 
