Sequence Recorder: Auto size sections after recording to expand to time range of recorded data.

diff --git a/Engine/Source/Editor/SequenceRecorder/Private/ActorRecording.cpp b/Engine/Source/Editor/SequenceRecorder/Private/ActorRecording.cpp
--- a/Engine/Source/Editor/SequenceRecorder/Private/ActorRecording.cpp
+++ b/Engine/Source/Editor/SequenceRecorder/Private/ActorRecording.cpp
@@ -797,7 +797,13 @@ bool UActorRecording::StopRecording(ULevelSequence* OriginalSequence, float Curr
 
 		FFrameNumber RecordStartTime = OriginalSequence->GetMovieScene()->GetPlaybackRange().GetLowerBoundValue();
 		UMovieSceneSubSection* SubSection = SubTrack->AddSequence(CurrentSequence, RecordStartTime,  MovieScene::DiscreteSize(CurrentSequence->GetMovieScene()->GetPlaybackRange()));
-		SubSection->SetRowIndex(RowIndex.GetValue());		
+
+		if (SubSection->GetAutoSizeRange().IsSet())
+		{
+			SubSection->SetRange(SubSection->GetAutoSizeRange().GetValue());
+		}
+
+		SubSection->SetRowIndex(RowIndex.GetValue());
 
 		SubTrack->FixRowIndices();
 
