Sequence Recorder: Fix updating the target level sequence when duplicating before recording into the level sequence.

diff --git a/Engine/Source/Editor/SequenceRecorder/Private/SequenceRecorder.cpp b/Engine/Source/Editor/SequenceRecorder/Private/SequenceRecorder.cpp
--- a/Engine/Source/Editor/SequenceRecorder/Private/SequenceRecorder.cpp
+++ b/Engine/Source/Editor/SequenceRecorder/Private/SequenceRecorder.cpp
@@ -671,7 +671,7 @@ bool FSequenceRecorder::StartRecordingInternal(UWorld* World)
 
 					FString NewSequenceName = SequenceRecorderUtils::MakeNewAssetName(PathToRecordTo, SequenceName);
 
-					CurrentSequence = CastChecked<ULevelSequence>(SequenceRecorderUtils::DuplicateAsset(PathToRecordTo, SequenceName, RecordingGroup.Get()->TargetLevelSequence));
+					LevelSequence = CastChecked<ULevelSequence>(SequenceRecorderUtils::DuplicateAsset(PathToRecordTo, SequenceName, RecordingGroup.Get()->TargetLevelSequence));
 					
 					RecordingGroup.Get()->TargetLevelSequence = LevelSequence;
 				}
