Sequence Recorder: Fix typo

diff --git a/Engine/Source/Editor/SequenceRecorder/Public/SequenceRecorderSettings.h b/Engine/Source/Editor/SequenceRecorder/Public/SequenceRecorderSettings.h
--- a/Engine/Source/Editor/SequenceRecorder/Public/SequenceRecorderSettings.h
+++ b/Engine/Source/Editor/SequenceRecorder/Public/SequenceRecorderSettings.h
@@ -111,7 +111,7 @@ public:
 	UPROPERTY(Config, EditAnywhere, Category = "Audio Recording")
 	FString AudioSubDirectory;
 
-	/** Whether to record nearby spawned actors. If an actor matches a class in the ActorFilter, this state will by bypassed. */
+	/** Whether to record nearby spawned actors. If an actor matches a class in the ActorFilter, this state will be bypassed. */
 	UPROPERTY(Config, EditAnywhere, Category = "Sequence Recording")
 	bool bRecordNearbySpawnedActors;
 
