When attempting to load animation to make sure it doesn't exist on disk, make sure we don't log warnings as we're expecting the animation may not exist.

diff --git a/Engine/Source/Editor/SequenceRecorder/Private/AnimationRecorder.cpp b/Engine/Source/Editor/SequenceRecorder/Private/AnimationRecorder.cpp
--- a/Engine/Source/Editor/SequenceRecorder/Private/AnimationRecorder.cpp
+++ b/Engine/Source/Editor/SequenceRecorder/Private/AnimationRecorder.cpp
@@ -146,7 +146,7 @@ bool FAnimationRecorder::TriggerRecordAnimation(USkeletalMeshComponent* Componen
 		Parent = CreatePackage(nullptr, *ValidatedAssetPath);
 	}
 
-	UObject* const Object = LoadObject<UObject>(Parent, *ValidatedAssetName, nullptr, LOAD_None, nullptr);
+	UObject* const Object = LoadObject<UObject>(Parent, *ValidatedAssetName, nullptr, LOAD_Quiet, nullptr);
 	// if object with same name exists, warn user
 	if (Object)
 	{
