Changed take and sequence recorders to use the correct property name rather than the property's display text for its property lookup

diff --git a/Engine/Plugins/VirtualProduction/Takes/Source/TakeTrackRecorders/Private/MovieSceneTrackPropertyRecorder.cpp b/Engine/Plugins/VirtualProduction/Takes/Source/TakeTrackRecorders/Private/MovieSceneTrackPropertyRecorder.cpp
--- a/Engine/Plugins/VirtualProduction/Takes/Source/TakeTrackRecorders/Private/MovieSceneTrackPropertyRecorder.cpp
+++ b/Engine/Plugins/VirtualProduction/Takes/Source/TakeTrackRecorders/Private/MovieSceneTrackPropertyRecorder.cpp
@@ -55,7 +55,7 @@ UMovieSceneSection* FMovieSceneTrackPropertyRecorder<bool>::AddSection(const FSt
 
 	if (Track)
 	{
-		Track->SetPropertyNameAndPath(*TrackDisplayName, Binding.GetPropertyPath());
+		Track->SetPropertyNameAndPath(Binding.GetPropertyName(), Binding.GetPropertyPath());
 
 		UMovieSceneBoolSection* Section = Cast<UMovieSceneBoolSection>(Track->CreateNewSection());
 
@@ -210,7 +210,7 @@ UMovieSceneSection* FMovieSceneTrackPropertyRecorder<uint8>::AddSection(const FS
 
 	if (Track)
 	{
-		Track->SetPropertyNameAndPath(*TrackDisplayName, Binding.GetPropertyPath());
+		Track->SetPropertyNameAndPath(Binding.GetPropertyName(), Binding.GetPropertyPath());
 
 		UMovieSceneByteSection* Section = Cast<UMovieSceneByteSection>(Track->CreateNewSection());
 
@@ -356,7 +356,7 @@ UMovieSceneSection* FMovieSceneTrackPropertyRecorderEnum::AddSection(const FStri
 
 	if (Track)
 	{
-		Track->SetPropertyNameAndPath(*TrackDisplayName, Binding.GetPropertyPath());
+		Track->SetPropertyNameAndPath(Binding.GetPropertyName(), Binding.GetPropertyPath());
 
 		UMovieSceneEnumSection* Section = Cast<UMovieSceneEnumSection>(Track->CreateNewSection());
 
@@ -537,7 +537,7 @@ UMovieSceneSection* FMovieSceneTrackPropertyRecorder<float>::AddSection(const FS
 
 	if (Track)
 	{
-		Track->SetPropertyNameAndPath(*TrackDisplayName, Binding.GetPropertyPath());
+		Track->SetPropertyNameAndPath(Binding.GetPropertyName(), Binding.GetPropertyPath());
 
 		UMovieSceneFloatSection* Section = Cast<UMovieSceneFloatSection>(Track->CreateNewSection());
 
@@ -686,7 +686,7 @@ UMovieSceneSection* FMovieSceneTrackPropertyRecorder<FColor>::AddSection(const F
 
 	if (Track)
 	{
-		Track->SetPropertyNameAndPath(*TrackDisplayName, Binding.GetPropertyPath());
+		Track->SetPropertyNameAndPath(Binding.GetPropertyName(), Binding.GetPropertyPath());
 
 		UMovieSceneColorSection* Section = Cast<UMovieSceneColorSection>(Track->CreateNewSection());
 
@@ -886,7 +886,7 @@ UMovieSceneSection* FMovieSceneTrackPropertyRecorder<FVector>::AddSection(const
 	if (Track)
 	{
 		Track->SetNumChannelsUsed(3);
-		Track->SetPropertyNameAndPath(*TrackDisplayName, Binding.GetPropertyPath());
+		Track->SetPropertyNameAndPath(Binding.GetPropertyName(), Binding.GetPropertyPath());
 
 		UMovieSceneVectorSection* Section = Cast<UMovieSceneVectorSection>(Track->CreateNewSection());
 
@@ -1063,7 +1063,7 @@ UMovieSceneSection* FMovieSceneTrackPropertyRecorder<int32>::AddSection(const FS
 
 	if (Track)
 	{
-		Track->SetPropertyNameAndPath(*TrackDisplayName, Binding.GetPropertyPath());
+		Track->SetPropertyNameAndPath(Binding.GetPropertyName(), Binding.GetPropertyPath());
 
 		UMovieSceneIntegerSection* Section = Cast<UMovieSceneIntegerSection>(Track->CreateNewSection());
 
@@ -1224,7 +1224,7 @@ UMovieSceneSection* FMovieSceneTrackPropertyRecorder<FString>::AddSection(const
 
 	if (Track)
 	{
-		Track->SetPropertyNameAndPath(*TrackDisplayName, Binding.GetPropertyPath());
+		Track->SetPropertyNameAndPath(Binding.GetPropertyName(), Binding.GetPropertyPath());
 
 		UMovieSceneStringSection* Section = Cast<UMovieSceneStringSection>(Track->CreateNewSection());
 
diff --git a/Engine/Source/Editor/SequenceRecorder/Private/Sections/MovieScenePropertyRecorder.cpp b/Engine/Source/Editor/SequenceRecorder/Private/Sections/MovieScenePropertyRecorder.cpp
--- a/Engine/Source/Editor/SequenceRecorder/Private/Sections/MovieScenePropertyRecorder.cpp
+++ b/Engine/Source/Editor/SequenceRecorder/Private/Sections/MovieScenePropertyRecorder.cpp
@@ -33,7 +33,6 @@ UMovieSceneSection* FMovieScenePropertyRecorder<bool>::AddSection(UObject* InObj
 		return nullptr;
 	}
 
-	FString TrackDisplayName = *Binding.GetProperty(*InObjectToRecord)->GetDisplayNameText().ToString();
 	FName TrackName = *Binding.GetPropertyPath();
 	UMovieSceneBoolTrack* Track = InMovieScene->FindTrack<UMovieSceneBoolTrack>(InGuid, TrackName);
 	if (!Track)
@@ -47,7 +46,7 @@ UMovieSceneSection* FMovieScenePropertyRecorder<bool>::AddSection(UObject* InObj
 
 	if (Track)
 	{
-		Track->SetPropertyNameAndPath(*TrackDisplayName, Binding.GetPropertyPath());
+		Track->SetPropertyNameAndPath(Binding.GetPropertyName(), Binding.GetPropertyPath());
 
 		UMovieSceneBoolSection* Section         = Cast<UMovieSceneBoolSection>(Track->CreateNewSection());
 
@@ -95,7 +94,6 @@ UMovieSceneSection* FMovieScenePropertyRecorder<uint8>::AddSection(UObject* InOb
 		return nullptr;
 	}
 
-	FString TrackDisplayName = *Binding.GetProperty(*InObjectToRecord)->GetDisplayNameText().ToString();
 	FName TrackName = *Binding.GetPropertyPath();
 	UMovieSceneByteTrack* Track = InMovieScene->FindTrack<UMovieSceneByteTrack>(InGuid, TrackName);
 	if (!Track)
@@ -109,7 +107,7 @@ UMovieSceneSection* FMovieScenePropertyRecorder<uint8>::AddSection(UObject* InOb
 
 	if (Track)
 	{
-		Track->SetPropertyNameAndPath(*TrackDisplayName, Binding.GetPropertyPath());
+		Track->SetPropertyNameAndPath(Binding.GetPropertyName(), Binding.GetPropertyPath());
 
 		UMovieSceneByteSection* Section = Cast<UMovieSceneByteSection>(Track->CreateNewSection());
 
@@ -155,7 +153,6 @@ UMovieSceneSection* FMovieScenePropertyRecorderEnum::AddSection(UObject* InObjec
 		return nullptr;
 	}
 
-	FString TrackDisplayName = *Binding.GetProperty(*InObjectToRecord)->GetDisplayNameText().ToString();
 	FName TrackName = *Binding.GetPropertyPath();
 	UMovieSceneEnumTrack* Track = InMovieScene->FindTrack<UMovieSceneEnumTrack>(InGuid, TrackName);
 	if (!Track)
@@ -169,7 +166,7 @@ UMovieSceneSection* FMovieScenePropertyRecorderEnum::AddSection(UObject* InObjec
 
 	if (Track)
 	{
-		Track->SetPropertyNameAndPath(*TrackDisplayName, Binding.GetPropertyPath());
+		Track->SetPropertyNameAndPath(Binding.GetPropertyName(), Binding.GetPropertyPath());
 
 		UMovieSceneEnumSection* Section = Cast<UMovieSceneEnumSection>(Track->CreateNewSection());
 
@@ -215,7 +212,6 @@ UMovieSceneSection* FMovieScenePropertyRecorder<float>::AddSection(UObject* InOb
 		return nullptr;
 	}
 
-	FString TrackDisplayName = *Binding.GetProperty(*InObjectToRecord)->GetDisplayNameText().ToString();
 	FName TrackName = *Binding.GetPropertyPath();
 	UMovieSceneFloatTrack* Track = InMovieScene->FindTrack<UMovieSceneFloatTrack>(InGuid, TrackName);
 	if (!Track)
@@ -229,7 +225,7 @@ UMovieSceneSection* FMovieScenePropertyRecorder<float>::AddSection(UObject* InOb
 
 	if (Track)
 	{
-		Track->SetPropertyNameAndPath(*TrackDisplayName, Binding.GetPropertyPath());
+		Track->SetPropertyNameAndPath(Binding.GetPropertyName(), Binding.GetPropertyPath());
 
 		UMovieSceneFloatSection* Section = Cast<UMovieSceneFloatSection>(Track->CreateNewSection());
 
@@ -280,7 +276,6 @@ UMovieSceneSection* FMovieScenePropertyRecorder<FColor>::AddSection(UObject* InO
 		return nullptr;
 	}
 
-	FString TrackDisplayName = *Binding.GetProperty(*InObjectToRecord)->GetDisplayNameText().ToString();
 	FName TrackName = *Binding.GetPropertyPath();
 	UMovieSceneColorTrack* Track = InMovieScene->FindTrack<UMovieSceneColorTrack>(InGuid, TrackName);
 	if (!Track)
@@ -294,7 +289,7 @@ UMovieSceneSection* FMovieScenePropertyRecorder<FColor>::AddSection(UObject* InO
 
 	if (Track)
 	{
-		Track->SetPropertyNameAndPath(*TrackDisplayName, Binding.GetPropertyPath());
+		Track->SetPropertyNameAndPath(Binding.GetPropertyName(), Binding.GetPropertyPath());
 
 		UMovieSceneColorSection* Section = Cast<UMovieSceneColorSection>(Track->CreateNewSection());
 
@@ -365,7 +360,6 @@ UMovieSceneSection* FMovieScenePropertyRecorder<FVector>::AddSection(UObject* In
 		return nullptr;
 	}
 
-	FString TrackDisplayName = *Binding.GetProperty(*InObjectToRecord)->GetDisplayNameText().ToString();
 	FName TrackName = *Binding.GetPropertyPath();
 	UMovieSceneVectorTrack* Track = InMovieScene->FindTrack<UMovieSceneVectorTrack>(InGuid, TrackName);
 	if (!Track)
@@ -380,7 +374,7 @@ UMovieSceneSection* FMovieScenePropertyRecorder<FVector>::AddSection(UObject* In
 	if (Track)
 	{
 		Track->SetNumChannelsUsed(3);
-		Track->SetPropertyNameAndPath(*TrackDisplayName, Binding.GetPropertyPath());
+		Track->SetPropertyNameAndPath(Binding.GetPropertyName(), Binding.GetPropertyPath());
 
 		UMovieSceneVectorSection* Section = Cast<UMovieSceneVectorSection>(Track->CreateNewSection());
 
