Sequence Recorder: Remove level sequence actors to trigger from the actor groups. This was mistakenly added in a previous merge. It still exists in the sequence recorder settings.

diff --git a/Engine/Source/Editor/SequenceRecorder/Public/SequenceRecorderActorGroup.h b/Engine/Source/Editor/SequenceRecorder/Public/SequenceRecorderActorGroup.h
--- a/Engine/Source/Editor/SequenceRecorder/Public/SequenceRecorderActorGroup.h
+++ b/Engine/Source/Editor/SequenceRecorder/Public/SequenceRecorderActorGroup.h
@@ -51,10 +51,6 @@ public:
 	UPROPERTY(EditAnywhere, AdvancedDisplay, Category = "Recording Groups", meta=(EditCondition = "bSpecifyTargetLevelSequence"))
 	bool bRecordTargetLevelSequenceLength;
 
-	/** Sequence actors to trigger playback on when recording starts (e.g. for recording synchronized sequences) */
-	UPROPERTY(EditAnywhere, Category = "Recording Groups")
-	TArray<class ALevelSequenceActor*> LevelSequenceActorsToTrigger;
-
 	/** A list of actor recordings in this group which contains both the actors to record as well as settings for each one. */
 	UPROPERTY(VisibleAnywhere, Category = "Recording Groups")
 	TArray<UActorRecording*> RecordedActors;
