Mac CIS, adding forward declare and include for UAnimSequence

diff --git a/Engine/Source/Editor/SequenceRecorder/Private/SequenceRecorderModule.cpp b/Engine/Source/Editor/SequenceRecorder/Private/SequenceRecorderModule.cpp
--- a/Engine/Source/Editor/SequenceRecorder/Private/SequenceRecorderModule.cpp
+++ b/Engine/Source/Editor/SequenceRecorder/Private/SequenceRecorderModule.cpp
@@ -45,6 +45,7 @@
 #include "Framework/Docking/LayoutExtender.h"
 #include "MovieSceneTimeHelpers.h"
 #include "SequenceRecorderUtils.h"
+#include "Animation/AnimSequence.h"
 
 #define LOCTEXT_NAMESPACE "SequenceRecorder"
 
diff --git a/Engine/Source/Editor/SequenceRecorder/Public/ISequenceRecorder.h b/Engine/Source/Editor/SequenceRecorder/Public/ISequenceRecorder.h
--- a/Engine/Source/Editor/SequenceRecorder/Public/ISequenceRecorder.h
+++ b/Engine/Source/Editor/SequenceRecorder/Public/ISequenceRecorder.h
@@ -13,6 +13,7 @@ class AActor;
 class ISequenceAudioRecorder;
 class ISequenceRecorderExtender;
 class USequenceRecordingBase;
+class UAnimSequence;
 
 DECLARE_MULTICAST_DELEGATE_OneParam(FOnRecordingStarted, class UMovieSceneSequence* /*Sequence*/);
 
