Sequence Recorder: Correctly find valid selected recordings

diff --git a/Engine/Source/Editor/SequenceRecorder/Private/SSequenceRecorder.cpp b/Engine/Source/Editor/SequenceRecorder/Private/SSequenceRecorder.cpp
--- a/Engine/Source/Editor/SequenceRecorder/Private/SSequenceRecorder.cpp
+++ b/Engine/Source/Editor/SequenceRecorder/Private/SSequenceRecorder.cpp
@@ -720,9 +720,30 @@ bool SSequenceRecorder::CanRemoveRecording() const
 	int32 NumItemsSelected = 0;
 	for (const TSharedPtr<SListView<USequenceRecordingBase*>>& ListView : ExtenderListViews)
 	{
-		NumItemsSelected += ListView->GetNumItemsSelected();
+		TArray<USequenceRecordingBase*> SelectedBaseRecordings;
+		ListView->GetSelectedItems(SelectedBaseRecordings);
+
+		for (USequenceRecordingBase* SelectedBaseRecording : SelectedBaseRecordings)
+		{
+			if (SelectedBaseRecording)
+			{
+				++NumItemsSelected;
+			}
+		}
 	}
-	return (ActorListView->GetNumItemsSelected() > 0 || NumItemsSelected > 0) && !FSequenceRecorder::Get().IsRecording() && !FAnimationRecorderManager::Get().IsRecording();
+
+	TArray<UActorRecording*> SelectedActorRecordings;
+	ActorListView->GetSelectedItems(SelectedActorRecordings);
+
+	for (UActorRecording* ActorRecording : SelectedActorRecordings)
+	{
+		if (ActorRecording)
+		{
+			++NumItemsSelected;
+		}
+	}
+
+	return (NumItemsSelected > 0) && !FSequenceRecorder::Get().IsRecording() && !FAnimationRecorderManager::Get().IsRecording();
 }
 
 void SSequenceRecorder::HandleRemoveAllRecordings()
