Build warning.
#trivial

diff --git a/Engine/Source/Editor/SequenceRecorder/Private/AnimationRecorder.cpp b/Engine/Source/Editor/SequenceRecorder/Private/AnimationRecorder.cpp
--- a/Engine/Source/Editor/SequenceRecorder/Private/AnimationRecorder.cpp
+++ b/Engine/Source/Editor/SequenceRecorder/Private/AnimationRecorder.cpp
@@ -603,15 +603,7 @@ bool FAnimationRecorder::Record(USkeletalMeshComponent* Component, FTransform co
 				FTransform LocalTransform = SpacesBases[BoneIndex];
 				if ( ParentIndex != INDEX_NONE )
 				{
-					if (ParentIndex == SkeletonRootIndex)
-					{
-						// Remove initial root transform
-						LocalTransform.SetToRelativeTransform(SpacesBases[ParentIndex]);
-					}
-					else
-					{
-						LocalTransform.SetToRelativeTransform(SpacesBases[ParentIndex]);
-					}
+					LocalTransform.SetToRelativeTransform(SpacesBases[ParentIndex]);
 				}
 				// if record local to world, we'd like to consider component to world to be in root
 				else
