For some reason we were talking wrong delta time when starting animation recorder. Also now we really make sure to keep timre rate the timecode rate if needed to match up with other static mesh.

diff --git a/Engine/Plugins/VirtualProduction/Takes/Source/TakeTrackRecorders/Private/MovieSceneAnimationTrackRecorder.cpp b/Engine/Plugins/VirtualProduction/Takes/Source/TakeTrackRecorders/Private/MovieSceneAnimationTrackRecorder.cpp
--- a/Engine/Plugins/VirtualProduction/Takes/Source/TakeTrackRecorders/Private/MovieSceneAnimationTrackRecorder.cpp
+++ b/Engine/Plugins/VirtualProduction/Takes/Source/TakeTrackRecorders/Private/MovieSceneAnimationTrackRecorder.cpp
@@ -244,7 +244,9 @@ void UMovieSceneAnimationTrackRecorder::RecordSampleImpl(const FQualifiedFrameTi
 			bRecordInWorldSpace = !OwningTakeRecorderSource->IsOtherActorBeingRecorded(AttachParent->GetOwner());
 		}
 
-		FFrameRate SampleRate = MovieSceneSection->GetTypedOuter<UMovieScene>()->GetDisplayRate();
+		const UTimecodeProvider* TimecodeProvider = GEngine->GetTimecodeProvider();
+		FFrameRate SampleRate = (TimecodeProvider && TimecodeProvider->GetSynchronizationState() == ETimecodeProviderSynchronizationState::Synchronized)
+			? TimecodeProvider->GetFrameRate() : MovieScene->GetDisplayRate();
 
 		//Set this up here so we know that it's parent sources have also been added so we record in the correct space
 		FAnimationRecordingSettings RecordingSettings;
@@ -254,7 +256,7 @@ void UMovieSceneAnimationTrackRecorder::RecordSampleImpl(const FQualifiedFrameTi
 		RecordingSettings.Length = 0;
 		RecordingSettings.bRecordInWorldSpace = bRecordInWorldSpace;
 		RecordingSettings.bRemoveRootAnimation = bRemoveRootAnimation;
-
+		RecordingSettings.bCheckDeltaTimeAtBeginning = false;
 		AnimationRecorder.Init(SkeletalMeshComponent.Get(), AnimSequence.Get(), &AnimationSerializer, RecordingSettings);
 		AnimationRecorder.BeginRecording();
 	}
diff --git a/Engine/Source/Editor/SequenceRecorder/Private/AnimationRecorder.cpp b/Engine/Source/Editor/SequenceRecorder/Private/AnimationRecorder.cpp
--- a/Engine/Source/Editor/SequenceRecorder/Private/AnimationRecorder.cpp
+++ b/Engine/Source/Editor/SequenceRecorder/Private/AnimationRecorder.cpp
@@ -29,6 +29,7 @@ FAnimationRecorder::FAnimationRecorder()
 	, bRecordLocalToWorld(false)
 	, bAutoSaveAsset(false)
 	, bRemoveRootTransform(true)
+	, bCheckDeltaTimeAtBeginning(true)
 	, InterpMode(ERichCurveInterpMode::RCIM_Linear)
 	, TangentMode(ERichCurveTangentMode::RCTM_Auto)
 	, AnimationSerializer(nullptr)
@@ -456,10 +457,14 @@ void FAnimationRecorder::UpdateRecord(USkeletalMeshComponent* Component, float D
 		return;
 	}
 
-	// in-editor we can get a long frame update because of the modal dialog used to pick paths
-	if(DeltaTime > IntervalTime && (LastFrame == 0 || LastFrame == 1))
+	// Take Recorder will turn this off, not sure if it's needed for persona animation recording or not.
+	if (bCheckDeltaTimeAtBeginning)
 	{
-		DeltaTime = IntervalTime;
+		// in-editor we can get a long frame update because of the modal dialog used to pick paths
+		if (DeltaTime > IntervalTime && (LastFrame == 0 || LastFrame == 1))
+		{
+			DeltaTime = IntervalTime;
+		}
 	}
 
 	float const PreviousTimePassed = TimePassed;
@@ -812,6 +817,7 @@ void FAnimRecorderInstance::InitInternal(USkeletalMeshComponent* InComponent, co
 	Recorder->SetAnimCompressionScheme(UAnimCompress_BitwiseCompressOnly::StaticClass());
 	Recorder->bAutoSaveAsset = Settings.bAutoSaveAsset;
 	Recorder->bRemoveRootTransform = Settings.bRemoveRootAnimation;
+	Recorder->bCheckDeltaTimeAtBeginning = Settings.bCheckDeltaTimeAtBeginning;
 	Recorder->AnimationSerializer = InAnimationSerializer;
 
 	if (InComponent)
diff --git a/Engine/Source/Editor/SequenceRecorder/Public/AnimationRecorder.h b/Engine/Source/Editor/SequenceRecorder/Public/AnimationRecorder.h
--- a/Engine/Source/Editor/SequenceRecorder/Public/AnimationRecorder.h
+++ b/Engine/Source/Editor/SequenceRecorder/Public/AnimationRecorder.h
@@ -84,6 +84,8 @@ public:
 	uint8 bAutoSaveAsset : 1;
 	/** If true, the root bone transform will be removed from all bone transforms */
 	uint8 bRemoveRootTransform : 1;
+	/** If true we check delta time at beginning of recording */
+	uint8 bCheckDeltaTimeAtBeginning : 1;
 	/** The interpolation mode for the recorded keys */
 	ERichCurveInterpMode InterpMode;
 	/** The tangent mode for the recorded keys*/
diff --git a/Engine/Source/Runtime/Engine/Classes/Animation/AnimationRecordingSettings.h b/Engine/Source/Runtime/Engine/Classes/Animation/AnimationRecordingSettings.h
--- a/Engine/Source/Runtime/Engine/Classes/Animation/AnimationRecordingSettings.h
+++ b/Engine/Source/Runtime/Engine/Classes/Animation/AnimationRecordingSettings.h
@@ -30,6 +30,7 @@ struct ENGINE_API FAnimationRecordingSettings
 		, Length((float)DefaultMaximumLength)
 		, InterpMode(ERichCurveInterpMode::RCIM_Linear)
 		, TangentMode(ERichCurveTangentMode::RCTM_Auto)
+		, bCheckDeltaTimeAtBeginning(true)
 	{}
 
 	/** Whether to record animation in world space, defaults to true */
@@ -59,4 +60,7 @@ struct ENGINE_API FAnimationRecordingSettings
 	/** Tangent mode for the recorded keys. */
 	UPROPERTY(EditAnywhere, Category = "Settings")
 	TEnumAsByte<ERichCurveTangentMode> TangentMode;
+
+	/** Whether to check DeltaTime at recording for pauses, turned off for TakeRecorder*/
+	bool bCheckDeltaTimeAtBeginning;
 };
