Merging //UE4/Main @ 6810573 to //UE4/Dev-Core (sbo182_DevCore)

(peafour-mergedown)

diff --git a/Engine/Plugins/FX/Niagara/Source/Niagara/Public/NiagaraTypes.h b/Engine/Plugins/FX/Niagara/Source/Niagara/Public/NiagaraTypes.h
--- a/Engine/Plugins/FX/Niagara/Source/Niagara/Public/NiagaraTypes.h
+++ b/Engine/Plugins/FX/Niagara/Source/Niagara/Public/NiagaraTypes.h
@@ -354,7 +354,7 @@ public:
 	UPROPERTY(EditAnywhere, Category = "Variable", DisplayName = "Property Metadata", meta = (ToolTip = "Property Metadata"))
 	TMap<FName, FString> PropertyMetaData;
 
-	UPROPERTY()
+	UPROPERTY(AdvancedDisplay, VisibleAnywhere, Category = "Variable")
 	bool bIsStaticSwitch;
 };
 
diff --git a/Engine/Plugins/FX/Niagara/Source/NiagaraEditor/Private/Customizations/NiagaraFunctionCallNodeDetails.cpp b/Engine/Plugins/FX/Niagara/Source/NiagaraEditor/Private/Customizations/NiagaraFunctionCallNodeDetails.cpp
--- a/Engine/Plugins/FX/Niagara/Source/NiagaraEditor/Private/Customizations/NiagaraFunctionCallNodeDetails.cpp
+++ b/Engine/Plugins/FX/Niagara/Source/NiagaraEditor/Private/Customizations/NiagaraFunctionCallNodeDetails.cpp
@@ -12,6 +12,7 @@
 #include "Widgets/Input/SComboBox.h"
 #include "Widgets/Input/SNumericEntryBox.h"
 #include "Widgets/Input/SCheckBox.h"
+#include "Widgets/Input/SEditableTextBox.h"
 
 #define LOCTEXT_NAMESPACE "NiagaraFunctionCallNodeDetails"
 
@@ -58,27 +59,53 @@ void FNiagaraFunctionCallNodeDetails::CustomizeDetails(IDetailLayoutBuilder& Det
 		]
 		.ValueContent()
 		[
-			SNew(SBox)
+			SNew(SVerticalBox)
+			+ SVerticalBox::Slot()
 			.Padding(FMargin(0.0f, 2.0f))
 			[
 				SNew(SCheckBox)
 				.IsChecked_Lambda([SwitchInput, this]()
 				{
-					return Node->PropagatedStaticSwitchParameters.Contains(SwitchInput) ? ECheckBoxState::Checked : ECheckBoxState::Unchecked;
+					return Node->FindPropagatedVariable(SwitchInput) ? ECheckBoxState::Checked : ECheckBoxState::Unchecked;
 				})
 				.OnCheckStateChanged_Lambda([SwitchInput, this](const ECheckBoxState NewState)
 				{
 					if (NewState == ECheckBoxState::Unchecked)
 					{
-						Node->PropagatedStaticSwitchParameters.Remove(SwitchInput);
+						Node->RemovePropagatedVariable(SwitchInput);
 					}
-					else if (NewState == ECheckBoxState::Checked)
+					else if (NewState == ECheckBoxState::Checked && !Node->FindPropagatedVariable(SwitchInput))
 					{
-						Node->PropagatedStaticSwitchParameters.AddUnique(SwitchInput);
+						Node->PropagatedStaticSwitchParameters.Emplace(SwitchInput);
 					}
 					Node->RefreshFromExternalChanges();
 				})
 			]
+			+ SVerticalBox::Slot()
+			.HAlign(HAlign_Fill)
+			[
+				SNew(SEditableTextBox)
+				.ToolTipText(LOCTEXT("NiagaraOverridePropagatedValueName_TooltipText", "When set, the name of the propagated switch value is changed to a different value."))
+				.HintText(LOCTEXT("NiagaraOverridePropagatedValueName_HintText", "Override name"))
+				.IsEnabled_Lambda([SwitchInput, this]() { return Node->FindPropagatedVariable(SwitchInput) != nullptr; })
+				.OnTextCommitted_Lambda([SwitchInput, this](const FText& NewText, ETextCommit::Type CommitType)
+				{
+					FNiagaraPropagatedVariable* Propagated = Node->FindPropagatedVariable(SwitchInput);
+					if (Propagated)
+					{
+						Propagated->PropagatedName = NewText.ToString();
+					}
+				})
+				.Text_Lambda([SwitchInput, this]()
+				{
+					FNiagaraPropagatedVariable* Propagated = Node->FindPropagatedVariable(SwitchInput);
+					if (Propagated)
+					{
+						return FText::FromString(Propagated->PropagatedName);
+					}
+					return FText();
+				})
+			]
 		];
 	}
 }
diff --git a/Engine/Plugins/FX/Niagara/Source/NiagaraEditor/Private/Customizations/NiagaraStaticSwitchNodeDetails.cpp b/Engine/Plugins/FX/Niagara/Source/NiagaraEditor/Private/Customizations/NiagaraStaticSwitchNodeDetails.cpp
--- a/Engine/Plugins/FX/Niagara/Source/NiagaraEditor/Private/Customizations/NiagaraStaticSwitchNodeDetails.cpp
+++ b/Engine/Plugins/FX/Niagara/Source/NiagaraEditor/Private/Customizations/NiagaraStaticSwitchNodeDetails.cpp
@@ -11,6 +11,7 @@
 #include "DetailWidgetRow.h"
 #include "Widgets/Input/SComboBox.h"
 #include "Widgets/Input/SNumericEntryBox.h"
+#include "Widgets/Input/SEditableTextBox.h"
 
 #define LOCTEXT_NAMESPACE "NiagaraStaticSwitchNodeDetails"
 
@@ -46,7 +47,34 @@ void FNiagaraStaticSwitchNodeDetails::CustomizeDetails(IDetailLayoutBuilder& Det
 		UpdateSelectionFromNode();
 		
 		IDetailCategoryBuilder& CategoryBuilder = DetailBuilder.EditCategory(SwitchCategoryName);		
-		FDetailWidgetRow& DropdownWidget = CategoryBuilder.AddCustomRow(LOCTEXT("NiagaraSwitchNodeDetailsFilterText", "Input parameter type"));
+		FDetailWidgetRow& NameWidget = CategoryBuilder.AddCustomRow(LOCTEXT("NiagaraSwitchNodeNameFilterText", "Input parameter name"));
+
+		NameWidget
+		.NameContent()
+		[
+			SNew(SBox)
+			.Padding(FMargin(0.0f, 2.0f))
+			[
+				SNew(STextBlock)
+				.TextStyle(FNiagaraEditorStyle::Get(), "NiagaraEditor.ParameterText")
+				.Text(LOCTEXT("NiagaraSwitchNodeNameText", "Input parameter name"))
+			]
+		]
+		.ValueContent()
+		[
+			SNew(SBox)
+			.Padding(FMargin(0.0f, 2.0f))
+			[
+				SNew(SEditableTextBox)
+				.Text(this, &FNiagaraStaticSwitchNodeDetails::GetParameterNameText)
+				.ToolTipText(LOCTEXT("NiagaraSwitchNodeNameTooltip", "This is the name of the parameter that is exposed to the user calling this function graph."))
+				.OnTextCommitted(this, &FNiagaraStaticSwitchNodeDetails::OnParameterNameCommited)
+				.SelectAllTextWhenFocused(true)
+				.RevertTextOnEscape(true)
+			]
+		];
+
+		FDetailWidgetRow& DropdownWidget = CategoryBuilder.AddCustomRow(LOCTEXT("NiagaraSwitchNodeTypeFilterText", "Input parameter type"));
 
 		DropdownWidget
 		.NameContent()
@@ -124,6 +152,7 @@ void FNiagaraStaticSwitchNodeDetails::OnSelectionChanged(TSharedPtr<SwitchDropdo
 		return;
 	}
 	
+	FNiagaraTypeDefinition OldType = Node->GetInputType();
 	if (SelectedDropdownItem == DropdownOptions[0])
 	{
 		Node->SwitchTypeData.SwitchType = ENiagaraStaticSwitchType::Bool;
@@ -137,7 +166,7 @@ void FNiagaraStaticSwitchNodeDetails::OnSelectionChanged(TSharedPtr<SwitchDropdo
 		Node->SwitchTypeData.SwitchType = ENiagaraStaticSwitchType::Enum;
 		Node->SwitchTypeData.Enum = SelectedDropdownItem->Enum;
 	}
-	Node->RefreshFromExternalChanges();
+	Node->OnSwitchParameterTypeChanged(OldType);
 }
 
 FText FNiagaraStaticSwitchNodeDetails::GetDropdownItemLabel() const
@@ -196,4 +225,17 @@ void FNiagaraStaticSwitchNodeDetails::IntOptionValueCommitted(int32 Value, EText
 	}
 }
 
+FText FNiagaraStaticSwitchNodeDetails::GetParameterNameText() const
+{
+	return Node.IsValid() ? FText::FromName(Node->InputParameterName) : FText();
+}
+
+void FNiagaraStaticSwitchNodeDetails::OnParameterNameCommited(const FText& InText, ETextCommit::Type InCommitType)
+{
+	if (Node.IsValid())
+	{
+		Node->ChangeSwitchParameterName(FName(*InText.ToString()));
+	}
+}
+
 #undef LOCTEXT_NAMESPACE
diff --git a/Engine/Plugins/FX/Niagara/Source/NiagaraEditor/Private/Customizations/NiagaraStaticSwitchNodeDetails.h b/Engine/Plugins/FX/Niagara/Source/NiagaraEditor/Private/Customizations/NiagaraStaticSwitchNodeDetails.h
--- a/Engine/Plugins/FX/Niagara/Source/NiagaraEditor/Private/Customizations/NiagaraStaticSwitchNodeDetails.h
+++ b/Engine/Plugins/FX/Niagara/Source/NiagaraEditor/Private/Customizations/NiagaraStaticSwitchNodeDetails.h
@@ -44,6 +44,10 @@ private:
 	TOptional<int32> GetIntOptionValue() const;
 	void IntOptionValueCommitted(int32 Value, ETextCommit::Type CommitInfo);
 
+	// parameter name option function
+	FText GetParameterNameText() const;
+	void OnParameterNameCommited(const FText& InText, ETextCommit::Type InCommitType);
+
 	TWeakObjectPtr<class UNiagaraNodeStaticSwitch> Node;
 	TArray<TSharedPtr<SwitchDropdownOption>> DropdownOptions;
 	TSharedPtr<SwitchDropdownOption> SelectedDropdownItem;
diff --git a/Engine/Plugins/FX/Niagara/Source/NiagaraEditor/Private/NiagaraEditorUtilities.cpp b/Engine/Plugins/FX/Niagara/Source/NiagaraEditor/Private/NiagaraEditorUtilities.cpp
--- a/Engine/Plugins/FX/Niagara/Source/NiagaraEditor/Private/NiagaraEditorUtilities.cpp
+++ b/Engine/Plugins/FX/Niagara/Source/NiagaraEditor/Private/NiagaraEditorUtilities.cpp
@@ -649,17 +649,18 @@ void FNiagaraEditorUtilities::SetStaticSwitchConstants(UNiagaraGraph* Graph, con
 		UNiagaraNodeFunctionCall* FunctionNode = Cast<UNiagaraNodeFunctionCall>(Node);
 		if (FunctionNode && FunctionNode->PropagatedStaticSwitchParameters.Num() > 0)
 		{
-			for (const FNiagaraVariable& SwitchValue : FunctionNode->PropagatedStaticSwitchParameters)
+			for (const FNiagaraPropagatedVariable& SwitchValue : FunctionNode->PropagatedStaticSwitchParameters)
 			{
-				UEdGraphPin* ValuePin = FunctionNode->FindPin(SwitchValue.GetName(), EGPD_Input);
+				UEdGraphPin* ValuePin = FunctionNode->FindPin(SwitchValue.SwitchParameter.GetName(), EGPD_Input);
 				if (!ValuePin)
 				{
 					continue;
 				}
 				ValuePin->DefaultValue = FString();
+				FName PinName = SwitchValue.PropagatedName.IsEmpty() ? ValuePin->GetFName() : FName(*SwitchValue.PropagatedName);
 				for (UEdGraphPin* InputPin : CallInputs)
 				{
-					if (InputPin->GetFName().IsEqual(ValuePin->GetFName()) && InputPin->PinType == ValuePin->PinType)
+					if (InputPin->GetFName().IsEqual(PinName) && InputPin->PinType == ValuePin->PinType)
 					{
 						ValuePin->DefaultValue = InputPin->DefaultValue;
 						break;
diff --git a/Engine/Plugins/FX/Niagara/Source/NiagaraEditor/Private/NiagaraGraph.cpp b/Engine/Plugins/FX/Niagara/Source/NiagaraEditor/Private/NiagaraGraph.cpp
--- a/Engine/Plugins/FX/Niagara/Source/NiagaraEditor/Private/NiagaraGraph.cpp
+++ b/Engine/Plugins/FX/Niagara/Source/NiagaraEditor/Private/NiagaraGraph.cpp
@@ -629,12 +629,21 @@ TArray<FNiagaraVariable> UNiagaraGraph::FindStaticSwitchInputs() const
 		UNiagaraNodeFunctionCall* FunctionNode = Cast<UNiagaraNodeFunctionCall>(Node);
 		if (FunctionNode)
 		{
-			for (FNiagaraVariable Variable : FunctionNode->PropagatedStaticSwitchParameters)
+			for (const FNiagaraPropagatedVariable& Propagated : FunctionNode->PropagatedStaticSwitchParameters)
 			{
+				FNiagaraVariable Variable = Propagated.SwitchParameter;
+				if (!Propagated.PropagatedName.IsEmpty())
+				{
+					Variable.SetName(FName(*Propagated.PropagatedName));
+				}
 				Result.AddUnique(Variable);
 			}			
 		}
 	}
+	Result.Sort([](const FNiagaraVariable& Left, const FNiagaraVariable& Right)
+	{
+		return Left.GetName() < Right.GetName();
+	});
 	return Result;
 }
 
@@ -1538,8 +1547,13 @@ void UNiagaraGraph::RefreshParameterReferences() const
 		}
 		else if (UNiagaraNodeFunctionCall* FunctionNode = Cast<UNiagaraNodeFunctionCall>(Node))
 		{
-			for (FNiagaraVariable Variable : FunctionNode->PropagatedStaticSwitchParameters)
+			for (const FNiagaraPropagatedVariable& Propagated : FunctionNode->PropagatedStaticSwitchParameters)
 			{
+				FNiagaraVariable Variable = Propagated.SwitchParameter;
+				if (!Propagated.PropagatedName.IsEmpty())
+				{
+					Variable.SetName(FName(*Propagated.PropagatedName));
+				}
 				AddStaticParameterReference(Variable, FunctionNode);
 			}
 		}
diff --git a/Engine/Plugins/FX/Niagara/Source/NiagaraEditor/Private/NiagaraNodeFunctionCall.cpp b/Engine/Plugins/FX/Niagara/Source/NiagaraEditor/Private/NiagaraNodeFunctionCall.cpp
--- a/Engine/Plugins/FX/Niagara/Source/NiagaraEditor/Private/NiagaraNodeFunctionCall.cpp
+++ b/Engine/Plugins/FX/Niagara/Source/NiagaraEditor/Private/NiagaraNodeFunctionCall.cpp
@@ -166,7 +166,7 @@ void UNiagaraNodeFunctionCall::AllocateDefaultPins()
 			{
 				NewPin->DefaultValue = PinDefaultValue;
 			}
-			NewPin->bDefaultValueIsIgnored = PropagatedStaticSwitchParameters.Contains(Input);
+			NewPin->bDefaultValueIsIgnored = FindPropagatedVariable(Input) != nullptr;
 		}
 
 		for (FNiagaraVariable& Output : Outputs)
@@ -556,7 +556,7 @@ bool UNiagaraNodeFunctionCall::RefreshFromExternalChanges()
 			{
 				if (InputVar.GetName().IsEqual(Pin->GetFName()))
 				{
-					Pin->bDefaultValueIsIgnored = PropagatedStaticSwitchParameters.Contains(InputVar);
+					Pin->bDefaultValueIsIgnored = FindPropagatedVariable(InputVar) != nullptr;
 					break;
 				}
 			}
@@ -767,6 +767,30 @@ UNiagaraNodeFunctionCall::FOnInputsChanged& UNiagaraNodeFunctionCall::OnInputsCh
 	return OnInputsChangedDelegate;
 }
 
+FNiagaraPropagatedVariable* UNiagaraNodeFunctionCall::FindPropagatedVariable(const FNiagaraVariable& Variable)
+{
+	for (FNiagaraPropagatedVariable& Propagated : PropagatedStaticSwitchParameters)
+	{
+		if (Propagated.SwitchParameter == Variable)
+		{
+			return &Propagated;
+		}
+	}
+	return nullptr;
+}
+
+void UNiagaraNodeFunctionCall::RemovePropagatedVariable(const FNiagaraVariable& Variable)
+{
+	for (int i = 0; i < PropagatedStaticSwitchParameters.Num(); i++)
+	{
+		if (PropagatedStaticSwitchParameters[i].SwitchParameter == Variable)
+		{
+			PropagatedStaticSwitchParameters.RemoveAt(i);
+			return;
+		}
+	}
+}
+
 ENiagaraNumericOutputTypeSelectionMode UNiagaraNodeFunctionCall::GetNumericOutputTypeSelectionMode() const
 {
 	if (FunctionScript)
diff --git a/Engine/Plugins/FX/Niagara/Source/NiagaraEditor/Private/NiagaraNodeStaticSwitch.cpp b/Engine/Plugins/FX/Niagara/Source/NiagaraEditor/Private/NiagaraNodeStaticSwitch.cpp
--- a/Engine/Plugins/FX/Niagara/Source/NiagaraEditor/Private/NiagaraNodeStaticSwitch.cpp
+++ b/Engine/Plugins/FX/Niagara/Source/NiagaraEditor/Private/NiagaraNodeStaticSwitch.cpp
@@ -28,6 +28,35 @@ FNiagaraTypeDefinition UNiagaraNodeStaticSwitch::GetInputType() const
 	return FNiagaraTypeDefinition();
 }
 
+void UNiagaraNodeStaticSwitch::ChangeSwitchParameterName(const FName& NewName)
+{
+	FNiagaraVariable OldValue(GetInputType(), InputParameterName);
+	InputParameterName = NewName;
+	VisualsChangedDelegate.Broadcast(this);
+	RemoveUnusedGraphParameter(OldValue);	
+}
+
+void UNiagaraNodeStaticSwitch::OnSwitchParameterTypeChanged(const FNiagaraTypeDefinition& OldType)
+{
+	RefreshFromExternalChanges();
+	VisualsChangedDelegate.Broadcast(this);
+	RemoveUnusedGraphParameter(FNiagaraVariable(OldType, InputParameterName));
+}
+
+void UNiagaraNodeStaticSwitch::RemoveUnusedGraphParameter(const FNiagaraVariable& OldParameter)
+{
+	TArray<FNiagaraVariable> GraphVariables = GetNiagaraGraph()->FindStaticSwitchInputs();
+	int Index = GraphVariables.Find(OldParameter);
+	if (Index == INDEX_NONE)
+	{
+		GetNiagaraGraph()->RemoveParameter(OldParameter);
+	}
+	else
+	{
+		GetNiagaraGraph()->NotifyGraphChanged();
+	}
+}
+
 void UNiagaraNodeStaticSwitch::AllocateDefaultPins()
 {
 	const UEdGraphSchema_Niagara* Schema = GetDefault<UEdGraphSchema_Niagara>();
diff --git a/Engine/Plugins/FX/Niagara/Source/NiagaraEditor/Private/NiagaraNodeStaticSwitch.h b/Engine/Plugins/FX/Niagara/Source/NiagaraEditor/Private/NiagaraNodeStaticSwitch.h
--- a/Engine/Plugins/FX/Niagara/Source/NiagaraEditor/Private/NiagaraNodeStaticSwitch.h
+++ b/Engine/Plugins/FX/Niagara/Source/NiagaraEditor/Private/NiagaraNodeStaticSwitch.h
@@ -41,7 +41,7 @@ class UNiagaraNodeStaticSwitch : public UNiagaraNodeUsageSelector
 
 public:
 
-	UPROPERTY(EditAnywhere, Category="Static Switch")
+	UPROPERTY()
 	FName InputParameterName;
 	
 	UPROPERTY()
@@ -57,6 +57,9 @@ public:
 
 	FNiagaraTypeDefinition GetInputType() const;
 
+	void ChangeSwitchParameterName(const FName& NewName);
+	void OnSwitchParameterTypeChanged(const FNiagaraTypeDefinition& OldType);
+
 	//~ Begin EdGraphNode Interface
 	virtual void AllocateDefaultPins() override;
 	virtual FText GetTooltipText() const override;
@@ -82,4 +85,6 @@ private:
 	bool GetVarIndex(class FHlslNiagaraTranslator* Translator, int32 InputPinCount, int32& VarIndexOut) const;
 
 	bool GetVarIndex(class FHlslNiagaraTranslator* Translator, int32 InputPinCount, int32 Value, int32& VarIndexOut) const;
+
+	void RemoveUnusedGraphParameter(const FNiagaraVariable& OldParameter);
 };
diff --git a/Engine/Plugins/FX/Niagara/Source/NiagaraEditor/Private/NiagaraScriptVariable.h b/Engine/Plugins/FX/Niagara/Source/NiagaraEditor/Private/NiagaraScriptVariable.h
--- a/Engine/Plugins/FX/Niagara/Source/NiagaraEditor/Private/NiagaraScriptVariable.h
+++ b/Engine/Plugins/FX/Niagara/Source/NiagaraEditor/Private/NiagaraScriptVariable.h
@@ -14,6 +14,6 @@ public:
 	UPROPERTY()
 	FNiagaraVariable Variable;
 
-	UPROPERTY(EditAnywhere, Category = Metadata)
+	UPROPERTY(EditAnywhere, Category = Metadata, meta=(NoResetToDefault, ShowOnlyInnerProperties))
 	FNiagaraVariableMetaData Metadata;
 };
\ No newline at end of file
diff --git a/Engine/Plugins/FX/Niagara/Source/NiagaraEditor/Public/NiagaraNodeFunctionCall.h b/Engine/Plugins/FX/Niagara/Source/NiagaraEditor/Public/NiagaraNodeFunctionCall.h
--- a/Engine/Plugins/FX/Niagara/Source/NiagaraEditor/Public/NiagaraNodeFunctionCall.h
+++ b/Engine/Plugins/FX/Niagara/Source/NiagaraEditor/Public/NiagaraNodeFunctionCall.h
@@ -11,6 +11,30 @@
 
 class UNiagaraScript;
 
+USTRUCT()
+struct FNiagaraPropagatedVariable
+{
+	GENERATED_USTRUCT_BODY()
+
+public:
+
+	FNiagaraPropagatedVariable() : FNiagaraPropagatedVariable(FNiagaraVariable()) {}
+
+	FNiagaraPropagatedVariable(FNiagaraVariable SwitchParameter) : SwitchParameter(SwitchParameter), PropagatedName(FString()) {}
+
+	UPROPERTY()
+	FNiagaraVariable SwitchParameter;
+
+	/** If set, then this overrides the name of the switch parameter when propagating. */
+	UPROPERTY()
+	FString PropagatedName;
+
+	bool operator==(const FNiagaraPropagatedVariable& Other)const
+	{
+		return SwitchParameter == Other.SwitchParameter;
+	}
+};
+
 UCLASS(MinimalAPI)
 class UNiagaraNodeFunctionCall : public UNiagaraNodeWithDynamicPins
 {
@@ -38,7 +62,7 @@ public:
 
 	/** All the input values the function propagates to the next higher caller instead of forcing the user to set them directly. */
 	UPROPERTY()
-	TArray<FNiagaraVariable> PropagatedStaticSwitchParameters;
+	TArray<FNiagaraPropagatedVariable> PropagatedStaticSwitchParameters;
 
 	bool ScriptIsValid() const;
 
@@ -88,6 +112,9 @@ public:
 
 	FOnInputsChanged& OnInputsChanged();
 
+	FNiagaraPropagatedVariable* FindPropagatedVariable(const FNiagaraVariable& Variable);
+	void RemovePropagatedVariable(const FNiagaraVariable& Variable);
+
 protected:
 	virtual bool IsValidPinToCompile(UEdGraphPin* Pin) const;
 
diff --git a/Engine/Plugins/VirtualProduction/Takes/Source/TakeTrackRecorders/Private/MovieSceneTrackPropertyRecorder.cpp b/Engine/Plugins/VirtualProduction/Takes/Source/TakeTrackRecorders/Private/MovieSceneTrackPropertyRecorder.cpp
--- a/Engine/Plugins/VirtualProduction/Takes/Source/TakeTrackRecorders/Private/MovieSceneTrackPropertyRecorder.cpp
+++ b/Engine/Plugins/VirtualProduction/Takes/Source/TakeTrackRecorders/Private/MovieSceneTrackPropertyRecorder.cpp
@@ -55,7 +55,7 @@ UMovieSceneSection* FMovieSceneTrackPropertyRecorder<bool>::AddSection(const FSt
 
 	if (Track)
 	{
-		Track->SetPropertyNameAndPath(*TrackDisplayName, Binding.GetPropertyPath());
+		Track->SetPropertyNameAndPath(Binding.GetPropertyName(), Binding.GetPropertyPath());
 
 		UMovieSceneBoolSection* Section = Cast<UMovieSceneBoolSection>(Track->CreateNewSection());
 
@@ -210,7 +210,7 @@ UMovieSceneSection* FMovieSceneTrackPropertyRecorder<uint8>::AddSection(const FS
 
 	if (Track)
 	{
-		Track->SetPropertyNameAndPath(*TrackDisplayName, Binding.GetPropertyPath());
+		Track->SetPropertyNameAndPath(Binding.GetPropertyName(), Binding.GetPropertyPath());
 
 		UMovieSceneByteSection* Section = Cast<UMovieSceneByteSection>(Track->CreateNewSection());
 
@@ -356,7 +356,7 @@ UMovieSceneSection* FMovieSceneTrackPropertyRecorderEnum::AddSection(const FStri
 
 	if (Track)
 	{
-		Track->SetPropertyNameAndPath(*TrackDisplayName, Binding.GetPropertyPath());
+		Track->SetPropertyNameAndPath(Binding.GetPropertyName(), Binding.GetPropertyPath());
 
 		UMovieSceneEnumSection* Section = Cast<UMovieSceneEnumSection>(Track->CreateNewSection());
 
@@ -537,7 +537,7 @@ UMovieSceneSection* FMovieSceneTrackPropertyRecorder<float>::AddSection(const FS
 
 	if (Track)
 	{
-		Track->SetPropertyNameAndPath(*TrackDisplayName, Binding.GetPropertyPath());
+		Track->SetPropertyNameAndPath(Binding.GetPropertyName(), Binding.GetPropertyPath());
 
 		UMovieSceneFloatSection* Section = Cast<UMovieSceneFloatSection>(Track->CreateNewSection());
 
@@ -686,7 +686,7 @@ UMovieSceneSection* FMovieSceneTrackPropertyRecorder<FColor>::AddSection(const F
 
 	if (Track)
 	{
-		Track->SetPropertyNameAndPath(*TrackDisplayName, Binding.GetPropertyPath());
+		Track->SetPropertyNameAndPath(Binding.GetPropertyName(), Binding.GetPropertyPath());
 
 		UMovieSceneColorSection* Section = Cast<UMovieSceneColorSection>(Track->CreateNewSection());
 
@@ -886,7 +886,7 @@ UMovieSceneSection* FMovieSceneTrackPropertyRecorder<FVector>::AddSection(const
 	if (Track)
 	{
 		Track->SetNumChannelsUsed(3);
-		Track->SetPropertyNameAndPath(*TrackDisplayName, Binding.GetPropertyPath());
+		Track->SetPropertyNameAndPath(Binding.GetPropertyName(), Binding.GetPropertyPath());
 
 		UMovieSceneVectorSection* Section = Cast<UMovieSceneVectorSection>(Track->CreateNewSection());
 
@@ -1063,7 +1063,7 @@ UMovieSceneSection* FMovieSceneTrackPropertyRecorder<int32>::AddSection(const FS
 
 	if (Track)
 	{
-		Track->SetPropertyNameAndPath(*TrackDisplayName, Binding.GetPropertyPath());
+		Track->SetPropertyNameAndPath(Binding.GetPropertyName(), Binding.GetPropertyPath());
 
 		UMovieSceneIntegerSection* Section = Cast<UMovieSceneIntegerSection>(Track->CreateNewSection());
 
@@ -1224,7 +1224,7 @@ UMovieSceneSection* FMovieSceneTrackPropertyRecorder<FString>::AddSection(const
 
 	if (Track)
 	{
-		Track->SetPropertyNameAndPath(*TrackDisplayName, Binding.GetPropertyPath());
+		Track->SetPropertyNameAndPath(Binding.GetPropertyName(), Binding.GetPropertyPath());
 
 		UMovieSceneStringSection* Section = Cast<UMovieSceneStringSection>(Track->CreateNewSection());
 
diff --git a/Engine/Source/Editor/SequenceRecorder/Private/Sections/MovieScenePropertyRecorder.cpp b/Engine/Source/Editor/SequenceRecorder/Private/Sections/MovieScenePropertyRecorder.cpp
--- a/Engine/Source/Editor/SequenceRecorder/Private/Sections/MovieScenePropertyRecorder.cpp
+++ b/Engine/Source/Editor/SequenceRecorder/Private/Sections/MovieScenePropertyRecorder.cpp
@@ -33,7 +33,6 @@ UMovieSceneSection* FMovieScenePropertyRecorder<bool>::AddSection(UObject* InObj
 		return nullptr;
 	}
 
-	FString TrackDisplayName = *Binding.GetProperty(*InObjectToRecord)->GetDisplayNameText().ToString();
 	FName TrackName = *Binding.GetPropertyPath();
 	UMovieSceneBoolTrack* Track = InMovieScene->FindTrack<UMovieSceneBoolTrack>(InGuid, TrackName);
 	if (!Track)
@@ -47,7 +46,7 @@ UMovieSceneSection* FMovieScenePropertyRecorder<bool>::AddSection(UObject* InObj
 
 	if (Track)
 	{
-		Track->SetPropertyNameAndPath(*TrackDisplayName, Binding.GetPropertyPath());
+		Track->SetPropertyNameAndPath(Binding.GetPropertyName(), Binding.GetPropertyPath());
 
 		UMovieSceneBoolSection* Section         = Cast<UMovieSceneBoolSection>(Track->CreateNewSection());
 
@@ -95,7 +94,6 @@ UMovieSceneSection* FMovieScenePropertyRecorder<uint8>::AddSection(UObject* InOb
 		return nullptr;
 	}
 
-	FString TrackDisplayName = *Binding.GetProperty(*InObjectToRecord)->GetDisplayNameText().ToString();
 	FName TrackName = *Binding.GetPropertyPath();
 	UMovieSceneByteTrack* Track = InMovieScene->FindTrack<UMovieSceneByteTrack>(InGuid, TrackName);
 	if (!Track)
@@ -109,7 +107,7 @@ UMovieSceneSection* FMovieScenePropertyRecorder<uint8>::AddSection(UObject* InOb
 
 	if (Track)
 	{
-		Track->SetPropertyNameAndPath(*TrackDisplayName, Binding.GetPropertyPath());
+		Track->SetPropertyNameAndPath(Binding.GetPropertyName(), Binding.GetPropertyPath());
 
 		UMovieSceneByteSection* Section = Cast<UMovieSceneByteSection>(Track->CreateNewSection());
 
@@ -155,7 +153,6 @@ UMovieSceneSection* FMovieScenePropertyRecorderEnum::AddSection(UObject* InObjec
 		return nullptr;
 	}
 
-	FString TrackDisplayName = *Binding.GetProperty(*InObjectToRecord)->GetDisplayNameText().ToString();
 	FName TrackName = *Binding.GetPropertyPath();
 	UMovieSceneEnumTrack* Track = InMovieScene->FindTrack<UMovieSceneEnumTrack>(InGuid, TrackName);
 	if (!Track)
@@ -169,7 +166,7 @@ UMovieSceneSection* FMovieScenePropertyRecorderEnum::AddSection(UObject* InObjec
 
 	if (Track)
 	{
-		Track->SetPropertyNameAndPath(*TrackDisplayName, Binding.GetPropertyPath());
+		Track->SetPropertyNameAndPath(Binding.GetPropertyName(), Binding.GetPropertyPath());
 
 		UMovieSceneEnumSection* Section = Cast<UMovieSceneEnumSection>(Track->CreateNewSection());
 
@@ -215,7 +212,6 @@ UMovieSceneSection* FMovieScenePropertyRecorder<float>::AddSection(UObject* InOb
 		return nullptr;
 	}
 
-	FString TrackDisplayName = *Binding.GetProperty(*InObjectToRecord)->GetDisplayNameText().ToString();
 	FName TrackName = *Binding.GetPropertyPath();
 	UMovieSceneFloatTrack* Track = InMovieScene->FindTrack<UMovieSceneFloatTrack>(InGuid, TrackName);
 	if (!Track)
@@ -229,7 +225,7 @@ UMovieSceneSection* FMovieScenePropertyRecorder<float>::AddSection(UObject* InOb
 
 	if (Track)
 	{
-		Track->SetPropertyNameAndPath(*TrackDisplayName, Binding.GetPropertyPath());
+		Track->SetPropertyNameAndPath(Binding.GetPropertyName(), Binding.GetPropertyPath());
 
 		UMovieSceneFloatSection* Section = Cast<UMovieSceneFloatSection>(Track->CreateNewSection());
 
@@ -280,7 +276,6 @@ UMovieSceneSection* FMovieScenePropertyRecorder<FColor>::AddSection(UObject* InO
 		return nullptr;
 	}
 
-	FString TrackDisplayName = *Binding.GetProperty(*InObjectToRecord)->GetDisplayNameText().ToString();
 	FName TrackName = *Binding.GetPropertyPath();
 	UMovieSceneColorTrack* Track = InMovieScene->FindTrack<UMovieSceneColorTrack>(InGuid, TrackName);
 	if (!Track)
@@ -294,7 +289,7 @@ UMovieSceneSection* FMovieScenePropertyRecorder<FColor>::AddSection(UObject* InO
 
 	if (Track)
 	{
-		Track->SetPropertyNameAndPath(*TrackDisplayName, Binding.GetPropertyPath());
+		Track->SetPropertyNameAndPath(Binding.GetPropertyName(), Binding.GetPropertyPath());
 
 		UMovieSceneColorSection* Section = Cast<UMovieSceneColorSection>(Track->CreateNewSection());
 
@@ -365,7 +360,6 @@ UMovieSceneSection* FMovieScenePropertyRecorder<FVector>::AddSection(UObject* In
 		return nullptr;
 	}
 
-	FString TrackDisplayName = *Binding.GetProperty(*InObjectToRecord)->GetDisplayNameText().ToString();
 	FName TrackName = *Binding.GetPropertyPath();
 	UMovieSceneVectorTrack* Track = InMovieScene->FindTrack<UMovieSceneVectorTrack>(InGuid, TrackName);
 	if (!Track)
@@ -380,7 +374,7 @@ UMovieSceneSection* FMovieScenePropertyRecorder<FVector>::AddSection(UObject* In
 	if (Track)
 	{
 		Track->SetNumChannelsUsed(3);
-		Track->SetPropertyNameAndPath(*TrackDisplayName, Binding.GetPropertyPath());
+		Track->SetPropertyNameAndPath(Binding.GetPropertyName(), Binding.GetPropertyPath());
 
 		UMovieSceneVectorSection* Section = Cast<UMovieSceneVectorSection>(Track->CreateNewSection());
 
diff --git a/Engine/Source/Runtime/Engine/Classes/Engine/World.h b/Engine/Source/Runtime/Engine/Classes/Engine/World.h
--- a/Engine/Source/Runtime/Engine/Classes/Engine/World.h
+++ b/Engine/Source/Runtime/Engine/Classes/Engine/World.h
@@ -111,7 +111,10 @@ private:
 public:
 	~FConstPawnIterator();
 
-	operator bool() const;
+	FConstPawnIterator(FConstPawnIterator&&);
+	FConstPawnIterator& operator=(FConstPawnIterator&&);
+
+	explicit operator bool() const;
 	FPawnIteratorObject operator*() const;
 	TUniquePtr<FPawnIteratorObject> operator->() const;
 
@@ -123,7 +126,7 @@ public:
 	FConstPawnIterator& operator--(int) { return *this; }
 
 private:
-	TActorIterator<APawn>* Iterator;
+	TUniquePtr<TActorIterator<APawn>> Iterator;
 
 	friend UWorld;
 };
diff --git a/Engine/Source/Runtime/Engine/Private/World.cpp b/Engine/Source/Runtime/Engine/Private/World.cpp
--- a/Engine/Source/Runtime/Engine/Private/World.cpp
+++ b/Engine/Source/Runtime/Engine/Private/World.cpp
@@ -4182,35 +4182,43 @@ FConstPawnIterator::FConstPawnIterator(UWorld* World)
 {
 }
 
-FConstPawnIterator::~FConstPawnIterator()
-{
-	delete Iterator;
-}
+// Operators defined in the cpp to ensure the definition of TActorIterator is known
+FConstPawnIterator::FConstPawnIterator(FConstPawnIterator&&) = default;
+FConstPawnIterator& FConstPawnIterator::operator=(FConstPawnIterator&&) = default;
+FConstPawnIterator::~FConstPawnIterator() = default;
 
  FConstPawnIterator::operator bool() const
 {
-	return (bool)*Iterator;
+	return Iterator.IsValid() && (bool)*Iterator;
 }
 
 FConstPawnIterator& FConstPawnIterator::operator++()
 {
+	checkf(Iterator.IsValid(), TEXT("FConstPawnIterator::operator++() - this iterator has been moved from and is now invalid."));
+
 	++(*Iterator);
 	return *this;
 }
 
 FConstPawnIterator& FConstPawnIterator::operator++(int)
 {
+	checkf(Iterator.IsValid(), TEXT("FConstPawnIterator::operator++(int) - this iterator has been moved from and is now invalid."));
+
 	++(*Iterator);
 	return *this;
 }
 
 FPawnIteratorObject FConstPawnIterator::operator*() const
 {
+	checkf(Iterator.IsValid(), TEXT("FConstPawnIterator::operator*() - this iterator has been moved from and is now invalid."));
+
 	return FPawnIteratorObject(**Iterator);
 }
 
 TUniquePtr<FPawnIteratorObject> FConstPawnIterator::operator->() const
 {
+	checkf(Iterator.IsValid(), TEXT("FConstPawnIterator::operator->() - this iterator has been moved from and is now invalid."));
+
 	return TUniquePtr<FPawnIteratorObject>(new FPawnIteratorObject(**Iterator));
 }
 
