Take Recorder / Animation: Improve logging - added time for creating animation asset and saving sequences.
With LogAnimation Log and LogTakesCore Log:

LogAnimation: Animation Recorder set keys in 0.10 seconds
LogTakesCore: Saved take1 in 4.84 seconds
LogTakesCore: Saved take2 in 0.16 seconds
LogTakesCore: Saved take3 in 0.01 seconds
LogTakesCore: Saved take4 in 0.00 seconds
LogTakesCore: Saved take5 in 0.01 seconds
LogTakesCore: Finished processing take in 6.83 seconds

diff --git a/Engine/Plugins/VirtualProduction/Takes/Source/TakeRecorder/Private/Recorder/TakeRecorder.cpp b/Engine/Plugins/VirtualProduction/Takes/Source/TakeRecorder/Private/Recorder/TakeRecorder.cpp
--- a/Engine/Plugins/VirtualProduction/Takes/Source/TakeRecorder/Private/Recorder/TakeRecorder.cpp
+++ b/Engine/Plugins/VirtualProduction/Takes/Source/TakeRecorder/Private/Recorder/TakeRecorder.cpp
@@ -796,6 +796,8 @@ void UTakeRecorder::Stop()
 		return;
 	}
 
+	double StartTime = FPlatformTime::Seconds();
+
 	TGuardValue<bool> ReentrantGuard(bStoppedRecording, true);
 
 	FTimecode TimecodeOut = FApp::GetTimecode();
@@ -919,6 +921,15 @@ void UTakeRecorder::Stop()
 		}
 	}
 
+	if (SequenceAsset)
+	{
+		UPackage* const Package = SequenceAsset->GetOutermost();
+		FString const PackageName = Package->GetName();
+
+		double ElapsedTime = FPlatformTime::Seconds() - StartTime;
+		UE_LOG(LogTakesCore, Log, TEXT("Finished processing %s in %0.2f seconds"), *PackageName, ElapsedTime);
+	}
+
 	RemoveFromRoot();
 }
 
diff --git a/Engine/Plugins/VirtualProduction/Takes/Source/TakeRecorderSources/Private/TakeRecorderActorSource.cpp b/Engine/Plugins/VirtualProduction/Takes/Source/TakeRecorderSources/Private/TakeRecorderActorSource.cpp
--- a/Engine/Plugins/VirtualProduction/Takes/Source/TakeRecorderSources/Private/TakeRecorderActorSource.cpp
+++ b/Engine/Plugins/VirtualProduction/Takes/Source/TakeRecorderSources/Private/TakeRecorderActorSource.cpp
@@ -1254,7 +1254,7 @@ void UTakeRecorderActorSource::GetSceneComponents(USceneComponent* OnSceneCompon
 	// so that each child gets added to the out array and their children recursively.
 	if (OnSceneComponent->ComponentHasTag(DoNotRecordTag))
 	{
-		UE_LOG(LogTakesCore, Warning, TEXT("Skipping record component: %s with do not record tag"), *OnSceneComponent->GetName());
+		UE_LOG(LogTakesCore, Verbose, TEXT("Skipping record component: %s with do not record tag"), *OnSceneComponent->GetName());
 		return;
 	}
 
diff --git a/Engine/Plugins/VirtualProduction/Takes/Source/TakesCore/Private/TakesUtils.cpp b/Engine/Plugins/VirtualProduction/Takes/Source/TakesCore/Private/TakesUtils.cpp
--- a/Engine/Plugins/VirtualProduction/Takes/Source/TakesCore/Private/TakesUtils.cpp
+++ b/Engine/Plugins/VirtualProduction/Takes/Source/TakesCore/Private/TakesUtils.cpp
@@ -1,6 +1,7 @@
 // Copyright Epic Games, Inc. All Rights Reserved.
 
 #include "TakesUtils.h"
+#include "TakesCoreLog.h"
 
 #include "Engine/Engine.h"
 #include "Engine/World.h"
@@ -91,8 +92,13 @@ void SaveAsset(UObject* InObject)
 	UPackage* const Package = InObject->GetOutermost();
 	FString const PackageName = Package->GetName();
 	FString const PackageFileName = FPackageName::LongPackageNameToFilename(PackageName, FPackageName::GetAssetPackageExtension());
+	
+	double StartTime = FPlatformTime::Seconds();
 
 	UPackage::SavePackage(Package, NULL, RF_Standalone, *PackageFileName, GError, nullptr, false, true, SAVE_NoError);
+
+	double ElapsedTime = FPlatformTime::Seconds() - StartTime;
+	UE_LOG(LogTakesCore, Log, TEXT("Saved %s in %0.2f seconds"), *PackageName, ElapsedTime);
 }
 
 void CreateCameraCutTrack(ULevelSequence* LevelSequence, const FGuid& RecordedCameraGuid, const FMovieSceneSequenceID& SequenceID, const TRange<FFrameNumber>& InRange)
diff --git a/Engine/Source/Editor/SequenceRecorder/Private/AnimationRecorder.cpp b/Engine/Source/Editor/SequenceRecorder/Private/AnimationRecorder.cpp
--- a/Engine/Source/Editor/SequenceRecorder/Private/AnimationRecorder.cpp
+++ b/Engine/Source/Editor/SequenceRecorder/Private/AnimationRecorder.cpp
@@ -284,6 +284,8 @@ void FAnimationRecorder::FixupNotifies()
 
 UAnimSequence* FAnimationRecorder::StopRecord(bool bShowMessage)
 {
+	double StartTime, ElapsedTime = 0;
+
 	if (AnimationObject)
 	{
 		int32 NumFrames = LastFrame  + 1;
@@ -300,6 +302,8 @@ UAnimSequence* FAnimationRecorder::StopRecord(bool bShowMessage)
 		// add to real curve data 
 		if (RecordedCurves.Num() == NumFrames && UIDToArrayIndexLUT)
 		{
+			StartTime = FPlatformTime::Seconds();
+
 			USkeleton* SkeletonObj = AnimationObject->GetSkeleton();
 			for (int32 CurveUID = 0; CurveUID < UIDToArrayIndexLUT->Num(); ++CurveUID)
 			{
@@ -359,6 +363,9 @@ UAnimSequence* FAnimationRecorder::StopRecord(bool bShowMessage)
 					}
 				}
 			}	
+
+			ElapsedTime = FPlatformTime::Seconds() - StartTime;
+			UE_LOG(LogAnimation, Log, TEXT("Animation Recorder set keys in %0.02f seconds"), ElapsedTime);
 		}
 
 		//AnimationObject->RawCurveData.RemoveRedundantKeys();
@@ -372,8 +379,13 @@ UAnimSequence* FAnimationRecorder::StopRecord(bool bShowMessage)
 			UPackage* const Package = AnimationObject->GetOutermost();
 			FString const PackageName = Package->GetName();
 			FString const PackageFileName = FPackageName::LongPackageNameToFilename(PackageName, FPackageName::GetAssetPackageExtension());
+			
+			StartTime = FPlatformTime::Seconds();
 
 			UPackage::SavePackage(Package, NULL, RF_Standalone, *PackageFileName, GError, nullptr, false, true, SAVE_NoError);
+
+			ElapsedTime = FPlatformTime::Seconds() - StartTime;
+			UE_LOG(LogAnimation, Log, TEXT("Animation Recorder saved %s in %0.2f seconds"), *PackageName, ElapsedTime);
 		}
 
 		UAnimSequence* ReturnObject = AnimationObject;
