Animation Recorder: Curves Only Recording Broken.
#fix Ensure that when the users opts to not record bone transform at least a single frame of (identity) data is available to populate the model with
#misc Changed anim recorder to accumulate track data internally and only set the keys once when recording is stopped to improve perf

diff --git a/Engine/Source/Editor/SequenceRecorder/Private/AnimationRecorder.cpp b/Engine/Source/Editor/SequenceRecorder/Private/AnimationRecorder.cpp
--- a/Engine/Source/Editor/SequenceRecorder/Private/AnimationRecorder.cpp
+++ b/Engine/Source/Editor/SequenceRecorder/Private/AnimationRecorder.cpp
@@ -254,7 +254,8 @@ void FAnimationRecorder::StartRecord(USkeletalMeshComponent* Component, UAnimSeq
 		{
 			// add tracks for the bone existing
 			const FName BoneTreeName = AnimSkeleton->GetReferenceSkeleton().GetBoneName(BoneTreeIndex);
-			Controller->AddBoneTrack(BoneTreeName);
+			Controller->AddBoneTrack(BoneTreeName);			
+			RawTracks.AddDefaulted();
 		}
 	}
 
@@ -389,6 +390,16 @@ UAnimSequence* FAnimationRecorder::StopRecord(bool bShowMessage)
 			UE_LOG(LogAnimation, Log, TEXT("Animation Recorder set keys in %0.02f seconds"), ElapsedTime);
 		}
 
+		// Populate bone tracks
+		const TArray<FBoneAnimationTrack>& BoneAnimationTracks = AnimationObject->GetDataModel()->GetBoneAnimationTracks();
+		for (int32 TrackIndex = 0; TrackIndex < BoneAnimationTracks.Num(); ++TrackIndex)
+		{
+			const FBoneAnimationTrack& AnimationTrack = BoneAnimationTracks[TrackIndex];
+			FRawAnimSequenceTrack& RawTrack = RawTracks[TrackIndex];
+
+			Controller->SetBoneTrackKeys(AnimationTrack.Name, RawTrack.PosKeys, RawTrack.RotKeys, RawTrack.ScaleKeys);
+		}
+
 		Controller->NotifyPopulated();
 		Controller->CloseBracket();
 
@@ -707,6 +718,8 @@ bool FAnimationRecorder::Record(USkeletalMeshComponent* Component, FTransform co
 		for (int32 TrackIndex = 0; TrackIndex < BoneAnimationTracks.Num(); ++TrackIndex)
 		{
 			const FBoneAnimationTrack& AnimationTrack = BoneAnimationTracks[TrackIndex];
+			FRawAnimSequenceTrack& RawTrack = RawTracks[TrackIndex];
+
 			// verify if this bone exists in skeleton
 			const int32 BoneTreeIndex = AnimationTrack.BoneTreeIndex;
 			if (BoneTreeIndex != INDEX_NONE)
@@ -727,28 +740,29 @@ bool FAnimationRecorder::Record(USkeletalMeshComponent* Component, FTransform co
 					}
 				}
 
-				// Expensive copy, might be worth accumulating track data internally and only setting the keys once when recording is stopped
-				FRawAnimSequenceTrack RawTrack = AnimationTrack.InternalTrackData;
-
 				if (bRecordTransforms)
 				{
 					RawTrack.PosKeys.Add(LocalTransform.GetTranslation());
 					RawTrack.RotKeys.Add(LocalTransform.GetRotation());
-					RawTrack.ScaleKeys.Add(LocalTransform.GetScale3D());
-  
-					Controller->SetBoneTrackKeys(AnimationTrack.Name, RawTrack.PosKeys, RawTrack.RotKeys, RawTrack.ScaleKeys);			
-  
+					RawTrack.ScaleKeys.Add(LocalTransform.GetScale3D());  
 					if (AnimationSerializer)
 					{
 						SerializedAnimation.AddTransform(TrackIndex, LocalTransform);
 					}
-				}
 
-				// verification
-				if (FrameToAdd != RawTrack.PosKeys.Num()-1)
+					// verification
+					if (FrameToAdd != RawTrack.PosKeys.Num() - 1)
+					{
+						UE_LOG(LogAnimation, Warning, TEXT("Mismatch in animation frames. Trying to record frame: %d, but only: %d frame(s) exist. Changing skeleton while recording is not supported."), FrameToAdd, RawTrack.PosKeys.Num());
+						return false;
+					}
+				}
+				else if (FrameToAdd == 0)
 				{
-					UE_LOG(LogAnimation, Warning, TEXT("Mismatch in animation frames. Trying to record frame: %d, but only: %d frame(s) exist. Changing skeleton while recording is not supported."), FrameToAdd, RawTrack.PosKeys.Num());
-					return false;
+					// Populate with identity keys
+					RawTrack.PosKeys.Add(FVector::ZeroVector);
+					RawTrack.RotKeys.Add(FQuat::Identity);
+					RawTrack.ScaleKeys.Add(FVector::OneVector);
 				}
 			}
 		}
diff --git a/Engine/Source/Editor/SequenceRecorder/Public/AnimationRecorder.h b/Engine/Source/Editor/SequenceRecorder/Public/AnimationRecorder.h
--- a/Engine/Source/Editor/SequenceRecorder/Public/AnimationRecorder.h
+++ b/Engine/Source/Editor/SequenceRecorder/Public/AnimationRecorder.h
@@ -130,6 +130,7 @@ private:
 
 	TArray<FBlendedCurve> RecordedCurves;
 	TArray<uint16> const * UIDToArrayIndexLUT;
+	TArray<FRawAnimSequenceTrack> RawTracks;
 };
 
 //////////////////////////////////////////////////////////////////////////
