Animation Recorder: Notify recording Broken.
#fix delay adding notify events until recording has finished and frame-rate and play length have been set, as changing these will try to update notify timings

diff --git a/Engine/Source/Editor/SequenceRecorder/Private/AnimationRecorder.cpp b/Engine/Source/Editor/SequenceRecorder/Private/AnimationRecorder.cpp
--- a/Engine/Source/Editor/SequenceRecorder/Private/AnimationRecorder.cpp
+++ b/Engine/Source/Editor/SequenceRecorder/Private/AnimationRecorder.cpp
@@ -266,10 +266,14 @@ void FAnimationRecorder::StartRecord(USkeletalMeshComponent* Component, UAnimSeq
 	Record(Component, PreviousComponentToWorld, PreviousSpacesBases, PreviousAnimCurves,  0);
 }
 
-void FAnimationRecorder::FixupNotifies()
+void FAnimationRecorder::ProcessNotifies()
 {
 	if (AnimationObject)
 	{
+		// Copy recorded notify events, animation its notify array should be empty at this point
+		ensure(AnimationObject->Notifies.Num() == 0);
+		AnimationObject->Notifies = RecordedNotifyEvents;
+		
 		// build notify tracks - first find how many tracks we want
 		for (FAnimNotifyEvent& Event : AnimationObject->Notifies)
 		{
@@ -312,7 +316,7 @@ UAnimSequence* FAnimationRecorder::StopRecord(bool bShowMessage)
 		const float Numerator = FloatDenominator / IntervalTime;
 		Controller->SetFrameRate(FFrameRate(Numerator, FloatDenominator));
 
-		FixupNotifies();
+		ProcessNotifies();
 
 		// post-process applies compression etc.
 		// @todo figure out why removing redundant keys is inconsistent
@@ -882,7 +886,7 @@ void FAnimationRecorder::RecordNotifies(USkeletalMeshComponent* Component, const
 					}
 				}
 
-				AnimationObject->Notifies.Add(NewEvent);
+				RecordedNotifyEvents.Add(NewEvent);
 				AddedThisFrame++;
 			}
 		}
diff --git a/Engine/Source/Editor/SequenceRecorder/Public/AnimationRecorder.h b/Engine/Source/Editor/SequenceRecorder/Public/AnimationRecorder.h
--- a/Engine/Source/Editor/SequenceRecorder/Public/AnimationRecorder.h
+++ b/Engine/Source/Editor/SequenceRecorder/Public/AnimationRecorder.h
@@ -52,6 +52,9 @@ private:
 	/** Unique notify states added to this sequence during recording */
 	TMap<UAnimNotifyState*, UAnimNotifyState*> UniqueNotifyStates;
 
+	/** Notify events recorded at any point, processed and inserted into animation when recording has finished */
+	TArray<FAnimNotifyEvent> RecordedNotifyEvents;
+
 	static float DefaultSampleRate;
 
 	/** Array of times recorded */
@@ -112,7 +115,7 @@ private:
 
 	void RecordNotifies(USkeletalMeshComponent* Component, const TArray<FAnimNotifyEventReference>& AnimNotifies, float DeltaTime, float RecordTime);
 
-	void FixupNotifies();
+	void ProcessNotifies();
 
 	// recording curve data 
 	struct FBlendedCurve
