Saving a Level Sequence with a Linked Anim Sequence removes Notifies
#fix reimplemented 4.26 fix with new controller functionality, fully reset model if desired otherwise just delete the bone data

diff --git a/Engine/Source/Editor/SequenceRecorder/Private/AnimationRecorder.cpp b/Engine/Source/Editor/SequenceRecorder/Private/AnimationRecorder.cpp
--- a/Engine/Source/Editor/SequenceRecorder/Private/AnimationRecorder.cpp
+++ b/Engine/Source/Editor/SequenceRecorder/Private/AnimationRecorder.cpp
@@ -237,8 +237,15 @@ void FAnimationRecorder::StartRecord(USkeletalMeshComponent* Component, UAnimSeq
 
 	Controller.OpenBracket(LOCTEXT("StartRecord_Bracket", "Starting Animation Recording"));
 
-	AnimationObject->DeleteNotifyTrackData();
-	Controller.ResetModel();
+	const bool bKeepNotifiesAndCurves = CVarKeepNotifyAndCurvesOnAnimationRecord->GetInt() == 0 ? false : true;
+	if (bKeepNotifiesAndCurves)
+	{
+		Controller.RemoveAllBoneTracks();
+	}
+	else
+	{
+		AnimationObject->ResetAnimation();
+	}
 
 	RecordedCurves.Reset();
 	RecordedTimes.Empty();
@@ -270,8 +277,7 @@ void FAnimationRecorder::ProcessNotifies()
 	if (AnimationObject)
 	{
 		// Copy recorded notify events, animation its notify array should be empty at this point
-		ensure(AnimationObject->Notifies.Num() == 0);
-		AnimationObject->Notifies = RecordedNotifyEvents;
+		AnimationObject->Notifies.Append(RecordedNotifyEvents);
 		
 		// build notify tracks - first find how many tracks we want
 		for (FAnimNotifyEvent& Event : AnimationObject->Notifies)
diff --git a/Engine/Source/Runtime/Engine/Private/Animation/AnimSequence.cpp b/Engine/Source/Runtime/Engine/Private/Animation/AnimSequence.cpp
--- a/Engine/Source/Runtime/Engine/Private/Animation/AnimSequence.cpp
+++ b/Engine/Source/Runtime/Engine/Private/Animation/AnimSequence.cpp
@@ -2661,7 +2661,7 @@ void UAnimSequence::ResetAnimation()
 
 	Controller->OpenBracket(LOCTEXT("ResetAnimation_Bracket", "Resetting Animation Sequence"));
 	{
-		DeleteBoneAnimationData();
+		Controller->ResetModel();
 		DeleteNotifyTrackData();
 	}
 	Controller->CloseBracket();
@@ -2696,7 +2696,7 @@ void UAnimSequence::DeleteNotifyTrackData()
 void UAnimSequence::DeleteBoneAnimationData()
 {
 	ValidateModel();
-	Controller->ResetModel();
+	Controller->RemoveAllBoneTracks();
 }
 #endif // WITH_EDITOR
 		
