Automation Fix; Bad merge.

diff --git a/Engine/Source/Editor/SequenceRecorder/Private/AnimationRecorder.cpp b/Engine/Source/Editor/SequenceRecorder/Private/AnimationRecorder.cpp
--- a/Engine/Source/Editor/SequenceRecorder/Private/AnimationRecorder.cpp
+++ b/Engine/Source/Editor/SequenceRecorder/Private/AnimationRecorder.cpp
@@ -1023,7 +1023,7 @@ bool FAnimRecorderInstance::BeginRecording()
 
 void FAnimRecorderInstance::Update(float DeltaTime)
 {
-	if (SkelComp.IsValid())
+	if (SkelComp.IsValid() == false)
 	{
 		UE_LOG(LogAnimation, Log, TEXT("Animation Recorder:  Update: SkelMeshComp not Valid, No Recording will occur."));
 		return;
