Fix Bake Animation Sequence Curve/Transform Export
Enables toggling of exporting morph targets / attribute / material curves separately
Fixes issue when Transforms are not exported: zeroed keys were added instead of ref pose

diff --git a/Engine/Plugins/Animation/ControlRig/Source/ControlRigEditor/Private/ControlRigEditorModule.cpp b/Engine/Plugins/Animation/ControlRig/Source/ControlRigEditor/Private/ControlRigEditorModule.cpp
--- a/Engine/Plugins/Animation/ControlRig/Source/ControlRigEditor/Private/ControlRigEditorModule.cpp
+++ b/Engine/Plugins/Animation/ControlRig/Source/ControlRigEditor/Private/ControlRigEditorModule.cpp
@@ -769,7 +769,9 @@ void FControlRigEditorModule::BakeToControlRig(UClass* ControlRigClass, UAnimSeq
 						FLevelSequenceAnimSequenceLinkItem LevelAnimLinkItem;
 						LevelAnimLinkItem.SkelTrackGuid = ActorTrackGuid;
 						LevelAnimLinkItem.PathToAnimSequence = FSoftObjectPath(AnimSequence);
-						LevelAnimLinkItem.bExportCurves = true; //mz todo to fix
+						LevelAnimLinkItem.bExportMorphTargets = true; //mz todo to fix
+						LevelAnimLinkItem.bExportAttributeCurves = true;
+						LevelAnimLinkItem.bExportMaterialCurves = true;
 						LevelAnimLinkItem.bExportTransforms = true;
 						LevelAnimLinkItem.bRecordInWorldSpace = false;
 						LevelAnimLink->AnimSequenceLinks.Add(LevelAnimLinkItem);
diff --git a/Engine/Source/Editor/MovieSceneTools/Private/MovieSceneToolHelpers.cpp b/Engine/Source/Editor/MovieSceneTools/Private/MovieSceneToolHelpers.cpp
--- a/Engine/Source/Editor/MovieSceneTools/Private/MovieSceneToolHelpers.cpp
+++ b/Engine/Source/Editor/MovieSceneTools/Private/MovieSceneToolHelpers.cpp
@@ -3986,7 +3986,9 @@ bool MovieSceneToolHelpers::ExportToAnimSequence(UAnimSequence* AnimSequence, UA
 		RecordingSettings.bRemoveRootAnimation = false;
 		RecordingSettings.bCheckDeltaTimeAtBeginning = false;
 		RecordingSettings.bRecordTransforms = ExportOptions->bExportTransforms;
-		RecordingSettings.bRecordCurves = ExportOptions->bExportCurves;
+		RecordingSettings.bRecordMorphTargets = ExportOptions->bExportMorphTargets;
+		RecordingSettings.bRecordAttributeCurves = ExportOptions->bExportAttributeCurves;
+		RecordingSettings.bRecordMaterialCurves = ExportOptions->bExportMaterialCurves;
 		RecordingSettings.bRecordInWorldSpace = ExportOptions->bRecordInWorldSpace;
 		AnimationRecorder.Init(SkelMeshComp, AnimSequence, nullptr, RecordingSettings);	
 		});
diff --git a/Engine/Source/Editor/MovieSceneTools/Private/TrackEditors/SkeletalAnimationTrackEditor.cpp b/Engine/Source/Editor/MovieSceneTools/Private/TrackEditors/SkeletalAnimationTrackEditor.cpp
--- a/Engine/Source/Editor/MovieSceneTools/Private/TrackEditors/SkeletalAnimationTrackEditor.cpp
+++ b/Engine/Source/Editor/MovieSceneTools/Private/TrackEditors/SkeletalAnimationTrackEditor.cpp
@@ -1198,16 +1198,22 @@ void FSkeletalAnimationTrackEditor::OnSequencerSaved(ISequencer& )
 					USkeletalMeshComponent* SkelMeshComp = AcquireSkeletalMeshFromObjectGuid(Item.SkelTrackGuid, GetSequencer());
 					if (AnimSequence && SkelMeshComp)
 					{
-						bool bSavedExportCurves = AnimSeqExportOption->bExportCurves;
-						bool bSavedExportTransforms = AnimSeqExportOption->bExportTransforms;
-						bool bSavedIncludeComponentTransform = AnimSeqExportOption->bRecordInWorldSpace;
-						AnimSeqExportOption->bExportCurves = Item.bExportCurves;
+						const bool bSavedExportMorphTargets = AnimSeqExportOption->bExportMorphTargets;
+						const bool bSavedExportAttributeCurves = AnimSeqExportOption->bExportAttributeCurves;
+						const bool bSavedExportMaterialCurves = AnimSeqExportOption->bExportMaterialCurves;
+						const bool bSavedExportTransforms = AnimSeqExportOption->bExportTransforms;
+						const bool bSavedIncludeComponentTransform = AnimSeqExportOption->bRecordInWorldSpace;
+						AnimSeqExportOption->bExportMorphTargets = Item.bExportMorphTargets;
+						AnimSeqExportOption->bExportAttributeCurves = Item.bExportAttributeCurves;
+						AnimSeqExportOption->bExportMaterialCurves = Item.bExportMaterialCurves;
 						AnimSeqExportOption->bExportTransforms = Item.bExportTransforms;
 						AnimSeqExportOption->bRecordInWorldSpace = Item.bRecordInWorldSpace;
 
 						bool bResult = MovieSceneToolHelpers::ExportToAnimSequence(AnimSequence, AnimSeqExportOption, MovieScene, SequencerPtr.Get(), SkelMeshComp, Template, RootToLocalTransform);
 
-						AnimSeqExportOption->bExportCurves = bSavedExportCurves;
+						AnimSeqExportOption->bExportMorphTargets = bSavedExportMorphTargets;
+						AnimSeqExportOption->bExportAttributeCurves = bSavedExportAttributeCurves;
+						AnimSeqExportOption->bExportMaterialCurves = bSavedExportMaterialCurves;
 						AnimSeqExportOption->bExportTransforms = bSavedExportTransforms;
 						AnimSeqExportOption->bRecordInWorldSpace = bSavedIncludeComponentTransform;
 					}
@@ -1356,7 +1362,9 @@ bool FSkeletalAnimationTrackEditor::CreateAnimationSequence(const TArray<UObject
 									}
 								}
 								LevelAnimLinkItem.PathToAnimSequence = FSoftObjectPath(AnimSequence);
-								LevelAnimLinkItem.bExportCurves = AnimSeqExportOption->bExportCurves;
+								LevelAnimLinkItem.bExportMorphTargets = AnimSeqExportOption->bExportMorphTargets;
+								LevelAnimLinkItem.bExportAttributeCurves = AnimSeqExportOption->bExportAttributeCurves;
+								LevelAnimLinkItem.bExportMaterialCurves = AnimSeqExportOption->bExportMaterialCurves;
 								LevelAnimLinkItem.bExportTransforms = AnimSeqExportOption->bExportTransforms;
 								LevelAnimLinkItem.bRecordInWorldSpace = AnimSeqExportOption->bRecordInWorldSpace;
 
@@ -1374,7 +1382,9 @@ bool FSkeletalAnimationTrackEditor::CreateAnimationSequence(const TArray<UObject
 						FLevelSequenceAnimSequenceLinkItem LevelAnimLinkItem;
 						LevelAnimLinkItem.SkelTrackGuid = Binding;
 						LevelAnimLinkItem.PathToAnimSequence = FSoftObjectPath(AnimSequence);
-						LevelAnimLinkItem.bExportCurves = AnimSeqExportOption->bExportCurves;
+						LevelAnimLinkItem.bExportMorphTargets = AnimSeqExportOption->bExportMorphTargets;
+						LevelAnimLinkItem.bExportAttributeCurves = AnimSeqExportOption->bExportAttributeCurves;
+						LevelAnimLinkItem.bExportMaterialCurves = AnimSeqExportOption->bExportMaterialCurves;
 						LevelAnimLinkItem.bExportTransforms = AnimSeqExportOption->bExportTransforms;
 						LevelAnimLinkItem.bRecordInWorldSpace = AnimSeqExportOption->bRecordInWorldSpace;
 
diff --git a/Engine/Source/Editor/SequenceRecorder/Private/AnimationRecorder.cpp b/Engine/Source/Editor/SequenceRecorder/Private/AnimationRecorder.cpp
--- a/Engine/Source/Editor/SequenceRecorder/Private/AnimationRecorder.cpp
+++ b/Engine/Source/Editor/SequenceRecorder/Private/AnimationRecorder.cpp
@@ -230,7 +230,6 @@ void FAnimationRecorder::GetBoneTransforms(USkeletalMeshComponent* Component, TA
 	}
 }
 
-
 void FAnimationRecorder::StartRecord(USkeletalMeshComponent* Component, UAnimSequence* InAnimationObject)
 {
 	TimePassed = 0.f;
@@ -345,10 +344,26 @@ UAnimSequence* FAnimationRecorder::StopRecord(bool bShowMessage)
 			USkeleton* SkeletonObj = AnimationObject->GetSkeleton();
 			for (int32 CurveUID = 0; CurveUID < UIDToArrayIndexLUT->Num(); ++CurveUID)
 			{
-				int32 CurveIndex = (*UIDToArrayIndexLUT)[CurveUID];
+				const int32 CurveIndex = (*UIDToArrayIndexLUT)[CurveUID];
 
 				if (CurveIndex != MAX_uint16)
 				{
+					// Skip curves which type is disabled in the recorder settings
+					if (const FCurveMetaData* CurveMetaData = SkeletonObj->GetCurveMetaData(CurveUID))
+					{
+						const bool bMorphTarget = CurveMetaData->Type.bMorphtarget;
+						const bool bMaterialCurve = CurveMetaData->Type.bMaterial;
+						const bool bAttributeCurve = !bMorphTarget && !bMaterialCurve;
+						
+						const bool bSkipCurve = (bMorphTarget && !bRecordMorphTargets) ||
+												(bAttributeCurve && !bRecordAttributeCurves) ||
+												(bMaterialCurve && !bRecordMaterialCurves);
+						if (bSkipCurve)
+						{
+							continue;
+						}
+					}
+
 					const FFloatCurve* FloatCurveData = nullptr;
 
 					TArray<float> TimesToRecord;
@@ -360,8 +375,8 @@ UAnimSequence* FAnimationRecorder::StopRecord(bool bShowMessage)
 					int32 WriteIndex = 0;
 					for (int32 KeyIndex = 0; KeyIndex < NumKeys; ++KeyIndex)
 					{
-						const float TimeToRecord = KeyIndex*IntervalTime;
-						if(RecordedCurves[KeyIndex].ValidCurveWeights[CurveIndex])
+						const float TimeToRecord = KeyIndex * IntervalTime;
+						if (RecordedCurves[KeyIndex].ValidCurveWeights[CurveIndex])
 						{
 							float CurCurveValue = RecordedCurves[KeyIndex].CurveWeights[CurveIndex];
 							if (!bSeenThisCurve)
@@ -376,7 +391,6 @@ UAnimSequence* FAnimationRecorder::StopRecord(bool bShowMessage)
 									const FAnimationCurveIdentifier CurveId(CurveName, ERawCurveTrackTypes::RCT_Float);
 									Controller.AddCurve(CurveId, AACF_DefaultCurve);
 									FloatCurveData = AnimationObject->GetDataModel()->FindFloatCurve(CurveId);
-
 								}
 							}
 
@@ -417,7 +431,7 @@ UAnimSequence* FAnimationRecorder::StopRecord(bool bShowMessage)
 		for (int32 TrackIndex = 0; TrackIndex < BoneAnimationTracks.Num(); ++TrackIndex)
 		{
 			const FBoneAnimationTrack& AnimationTrack = BoneAnimationTracks[TrackIndex];
-			FRawAnimSequenceTrack& RawTrack = RawTracks[TrackIndex];
+			const FRawAnimSequenceTrack& RawTrack = RawTracks[TrackIndex];
 
 			Controller.SetBoneTrackKeys(AnimationTrack.Name, RawTrack.PosKeys, RawTrack.RotKeys, RawTrack.ScaleKeys);
 		}
@@ -648,7 +662,7 @@ void FAnimationRecorder::UpdateRecord(USkeletalMeshComponent* Component, float D
 		{
 			// find what frames we need to record
 			// convert to time
-			float CurrentTime = (FramesRecorded + 1) * IntervalTime;
+			const float CurrentTime = (FramesRecorded + 1) * IntervalTime;
 			float BlendAlpha = (CurrentTime - PreviousTimePassed) / DeltaTime;
 
 			UE_LOG(LogAnimation, Log, TEXT("Current Frame Count : %d, BlendAlpha : %0.2f"), FramesRecorded + 1, BlendAlpha);
@@ -761,22 +775,23 @@ bool FAnimationRecorder::Record(USkeletalMeshComponent* Component, FTransform co
 			{
 				const int32 BoneIndex = AnimSkeleton->GetMeshBoneIndexFromSkeletonBoneIndex(SkeletalMesh, BoneTreeIndex);
 				const int32 ParentIndex = SkeletalMesh->GetRefSkeleton().GetParentIndex(BoneIndex);
-				FTransform LocalTransform = SpacesBases[BoneIndex];
-				if ( ParentIndex != INDEX_NONE )
-				{
-					LocalTransform.SetToRelativeTransform(SpacesBases[ParentIndex]);
-				}
-				// if record local to world, we'd like to consider component to world to be in root
-				else
-				{
-					if (bRecordLocalToWorld)
-					{
-						LocalTransform *= ComponentToWorld;
-					}
-				}
 
 				if (bRecordTransforms)
 				{
+					FTransform LocalTransform = SpacesBases[BoneIndex];
+					if (ParentIndex != INDEX_NONE)
+					{
+						LocalTransform.SetToRelativeTransform(SpacesBases[ParentIndex]);
+					}
+					// if record local to world, we'd like to consider component to world to be in root
+					else
+					{
+						if (bRecordLocalToWorld)
+						{
+							LocalTransform *= ComponentToWorld;
+						}
+					}
+				
 					RawTrack.PosKeys.Add(FVector3f(LocalTransform.GetTranslation()));
 					RawTrack.RotKeys.Add(FQuat4f(LocalTransform.GetRotation()));
 					RawTrack.ScaleKeys.Add(FVector3f(LocalTransform.GetScale3D()));  
@@ -784,7 +799,6 @@ bool FAnimationRecorder::Record(USkeletalMeshComponent* Component, FTransform co
 					{
 						SerializedAnimation.AddTransform(TrackIndex, LocalTransform);
 					}
-
 					// verification
 					if (FrameToAdd != RawTrack.PosKeys.Num() - 1)
 					{
@@ -792,12 +806,15 @@ bool FAnimationRecorder::Record(USkeletalMeshComponent* Component, FTransform co
 						return false;
 					}
 				}
-				else if (FrameToAdd == 0)
+				else
 				{
-					// Populate with identity keys
-					RawTrack.PosKeys.Add(FVector3f::ZeroVector);
-					RawTrack.RotKeys.Add(FQuat4f::Identity);
-					RawTrack.ScaleKeys.Add(FVector3f::OneVector);
+					if (FrameToAdd == 0)
+					{
+						const FTransform RefPose = Component->SkeletalMesh->GetRefSkeleton().GetRefBonePose()[BoneIndex];
+						RawTrack.PosKeys.Add(RefPose.GetTranslation());
+						RawTrack.RotKeys.Add(FQuat4f(RefPose.GetRotation()));
+						RawTrack.ScaleKeys.Add(RefPose.GetScale3D());
+					}
 				}
 			}
 		}
@@ -810,6 +827,7 @@ bool FAnimationRecorder::Record(USkeletalMeshComponent* Component, FTransform co
 			AnimationSerializer->WriteFrameData(AnimationSerializer->FramesWritten, SerializedAnimation);
 		}
 		// each RecordedCurves contains all elements
+		const bool bRecordCurves = bRecordMorphTargets || bRecordAttributeCurves || bRecordMaterialCurves; 
 		if (bRecordCurves && AnimationCurves.CurveWeights.Num() > 0)
 		{
 			RecordedCurves.Emplace(AnimationCurves.CurveWeights, AnimationCurves.ValidCurveWeights);
@@ -994,7 +1012,9 @@ void FAnimRecorderInstance::InitInternal(USkeletalMeshComponent* InComponent, co
 	Recorder->bCheckDeltaTimeAtBeginning = Settings.bCheckDeltaTimeAtBeginning;
 	Recorder->AnimationSerializer = InAnimationSerializer;
 	Recorder->bRecordTransforms = Settings.bRecordTransforms;
-	Recorder->bRecordCurves = Settings.bRecordCurves;
+	Recorder->bRecordMorphTargets = Settings.bRecordMorphTargets;
+	Recorder->bRecordAttributeCurves = Settings.bRecordAttributeCurves;
+	Recorder->bRecordMaterialCurves = Settings.bRecordMaterialCurves;
 
 	if (InComponent)
 	{
diff --git a/Engine/Source/Editor/SequenceRecorder/Private/SequenceRecorderModule.cpp b/Engine/Source/Editor/SequenceRecorder/Private/SequenceRecorderModule.cpp
--- a/Engine/Source/Editor/SequenceRecorder/Private/SequenceRecorderModule.cpp
+++ b/Engine/Source/Editor/SequenceRecorder/Private/SequenceRecorderModule.cpp
@@ -785,10 +785,12 @@ class FSequenceRecorderModule : public ISequenceRecorder, private FSelfRegisteri
 	{
 		FAnimationRecorder Recorder;
 		Recorder.bRecordTransforms = true;
-		Recorder.bRecordCurves = true;
+		Recorder.bRecordMorphTargets = true;
+		Recorder.bRecordAttributeCurves = true;
+		Recorder.bRecordMaterialCurves = true;
 		if (Recorder.TriggerRecordAnimation(Component))
 		{
-			class UAnimSequence * Sequence = Recorder.GetAnimationObject();
+			class UAnimSequence* Sequence = Recorder.GetAnimationObject();
 			if (Sequence)
 			{
 				Recorder.StopRecord(false);
diff --git a/Engine/Source/Editor/SequenceRecorder/Public/AnimationRecorder.h b/Engine/Source/Editor/SequenceRecorder/Public/AnimationRecorder.h
--- a/Engine/Source/Editor/SequenceRecorder/Public/AnimationRecorder.h
+++ b/Engine/Source/Editor/SequenceRecorder/Public/AnimationRecorder.h
@@ -108,8 +108,12 @@ public:
 	FAnimationSerializer* AnimationSerializer;
 	/** Whether or not to record transforms*/
 	uint8 bRecordTransforms : 1;
-	/** Whether or not to record curves*/
-	uint8 bRecordCurves : 1;
+	/** Whether or not to record morph targets*/
+	uint8 bRecordMorphTargets : 1;
+	/** Whether or not to record attribute curves*/
+	uint8 bRecordAttributeCurves : 1;
+	/** Whether or not to record material curves*/
+	uint8 bRecordMaterialCurves : 1;
 public:
 	/** Helper function to get space bases depending on master pose component */
 	static void GetBoneTransforms(USkeletalMeshComponent* Component, TArray<FTransform>& BoneTransforms);
diff --git a/Engine/Source/Editor/UnrealEd/Classes/Exporters/AnimSeqExportOption.h b/Engine/Source/Editor/UnrealEd/Classes/Exporters/AnimSeqExportOption.h
--- a/Engine/Source/Editor/UnrealEd/Classes/Exporters/AnimSeqExportOption.h
+++ b/Engine/Source/Editor/UnrealEd/Classes/Exporters/AnimSeqExportOption.h
@@ -19,9 +19,17 @@ public:
 	UPROPERTY(EditAnywhere, BlueprintReadWrite, AdvancedDisplay, Category = Export)
 	bool bExportTransforms = true;
 
-	/** If enabled, export the curves from the animation */
+	/** If enabled, export the morph targets from the animation */
 	UPROPERTY(EditAnywhere, BlueprintReadWrite, AdvancedDisplay, Category = Export)
-	bool bExportCurves = true;
+	bool bExportMorphTargets = true;
+
+	/** If enabled, export the attribute curves from the animation */
+	UPROPERTY(EditAnywhere, BlueprintReadWrite, AdvancedDisplay, Category = Export)
+	bool bExportAttributeCurves = true;
+
+	/** If enabled, export the material curves from the animation */
+	UPROPERTY(EditAnywhere, BlueprintReadWrite, AdvancedDisplay, Category = Export)
+	bool bExportMaterialCurves = true;
 
 	/** If enabled we record in World Space otherwise we record from 0,0,0 */
 	UPROPERTY(EditAnywhere, BlueprintReadWrite, AdvancedDisplay, Category = Export)
@@ -43,7 +51,9 @@ public:
 	void ResetToDefault()
 	{
 		bExportTransforms = true;
-		bExportCurves = true;
+		bExportMorphTargets = true;
+		bExportAttributeCurves = true;
+		bExportMaterialCurves = true;
 		bRecordInWorldSpace = false;
 		bEvaluateAllSkeletalMeshComponents = false;
 		WarmUpFrames = 0;
diff --git a/Engine/Source/Runtime/Engine/Classes/Animation/AnimationRecordingSettings.h b/Engine/Source/Runtime/Engine/Classes/Animation/AnimationRecordingSettings.h
--- a/Engine/Source/Runtime/Engine/Classes/Animation/AnimationRecordingSettings.h
+++ b/Engine/Source/Runtime/Engine/Classes/Animation/AnimationRecordingSettings.h
@@ -32,7 +32,9 @@ struct ENGINE_API FAnimationRecordingSettings
 		, TangentMode(ERichCurveTangentMode::RCTM_Auto)
 		, bCheckDeltaTimeAtBeginning(true)
 		, bRecordTransforms(true)
-		, bRecordCurves(true)
+		, bRecordMorphTargets(true)
+		, bRecordAttributeCurves(true)
+		, bRecordMaterialCurves(true)
 	{}
 	
 	/** Whether to record animation in world space, defaults to true */
@@ -66,11 +68,19 @@ struct ENGINE_API FAnimationRecordingSettings
 	/** Whether to check DeltaTime at recording for pauses, turned off for TakeRecorder*/
 	bool bCheckDeltaTimeAtBeginning;
 
-	/** Whether or not to record transforms*/
+	/** Whether or not to record transforms */
 	UPROPERTY(EditAnywhere, Category = "Settings")
 	bool bRecordTransforms;
 
-	/** Whether or not to record curves*/
+	/** Whether or not to record morph targets */
 	UPROPERTY(EditAnywhere, Category = "Settings")
-	bool bRecordCurves;
+	bool bRecordMorphTargets;
+
+	/** Whether or not to record parameter curves */
+	UPROPERTY(EditAnywhere, Category = "Settings")
+	bool bRecordAttributeCurves;
+
+	/** Whether or not to record material curves */
+	UPROPERTY(EditAnywhere, Category = "Settings")
+	bool bRecordMaterialCurves;
 };
diff --git a/Engine/Source/Runtime/LevelSequence/Public/LevelSequenceAnimSequenceLink.h b/Engine/Source/Runtime/LevelSequence/Public/LevelSequenceAnimSequenceLink.h
--- a/Engine/Source/Runtime/LevelSequence/Public/LevelSequenceAnimSequenceLink.h
+++ b/Engine/Source/Runtime/LevelSequence/Public/LevelSequenceAnimSequenceLink.h
@@ -24,7 +24,11 @@ public:
 	UPROPERTY()
 	bool bExportTransforms = true;
 	UPROPERTY()
-	bool bExportCurves = true;
+	bool bExportMorphTargets = true;
+	UPROPERTY()
+	bool bExportAttributeCurves = true;
+	UPROPERTY()
+	bool bExportMaterialCurves = true;
 	UPROPERTY()
 	bool bRecordInWorldSpace = false;
 
