Frame Rate value of recorded animation doesn't match the user's set value
#fix replaced recording sample rate with FFrameRate to improve UX and integrate better with AnimDataModel

diff --git a/Engine/Plugins/VirtualProduction/Takes/Source/TakeTrackRecorders/Private/MovieSceneAnimationTrackRecorder.cpp b/Engine/Plugins/VirtualProduction/Takes/Source/TakeTrackRecorders/Private/MovieSceneAnimationTrackRecorder.cpp
--- a/Engine/Plugins/VirtualProduction/Takes/Source/TakeTrackRecorders/Private/MovieSceneAnimationTrackRecorder.cpp
+++ b/Engine/Plugins/VirtualProduction/Takes/Source/TakeTrackRecorders/Private/MovieSceneAnimationTrackRecorder.cpp
@@ -137,7 +137,8 @@ void UMovieSceneAnimationTrackRecorder::CreateTrackImpl()
 			FString Name = SkeletalMeshComponent->GetName();
 			FName SerializedType("Animation");
 			FString FileName = FString::Printf(TEXT("%s_%s"), *(SerializedType.ToString()), *Name);
-			float IntervalTime = SampleRate.AsDecimal() > 0.0f ? 1.0f / SampleRate.AsDecimal() : 1.0f / FAnimationRecordingSettings::DefaultSampleRate;
+
+			float IntervalTime = SampleRate.AsDecimal() > 0.0f ? 1.0f / SampleRate.AsDecimal() : FAnimationRecordingSettings::DefaultSampleFrameRate.AsInterval();
 			FAnimationFileHeader Header(SerializedType, ObjectGuid, IntervalTime);
 
 			USkeleton* AnimSkeleton = AnimSequence->GetSkeleton();
@@ -281,7 +282,7 @@ void UMovieSceneAnimationTrackRecorder::RecordSampleImpl(const FQualifiedFrameTi
 
 		//Set this up here so we know that it's parent sources have also been added so we record in the correct space
 		FAnimationRecordingSettings RecordingSettings;
-		RecordingSettings.SampleRate = SampleRate.AsDecimal();
+		RecordingSettings.SampleFrameRate = SampleRate;
 		RecordingSettings.InterpMode = AnimSettings->InterpMode;
 		RecordingSettings.TangentMode = AnimSettings->TangentMode;
 		RecordingSettings.Length = 0;
diff --git a/Engine/Source/Editor/MovieSceneTools/Private/MovieSceneToolHelpers.cpp b/Engine/Source/Editor/MovieSceneTools/Private/MovieSceneToolHelpers.cpp
--- a/Engine/Source/Editor/MovieSceneTools/Private/MovieSceneToolHelpers.cpp
+++ b/Engine/Source/Editor/MovieSceneTools/Private/MovieSceneToolHelpers.cpp
@@ -3979,7 +3979,7 @@ bool MovieSceneToolHelpers::ExportToAnimSequence(UAnimSequence* AnimSequence, UA
 	FInitAnimationCB InitCallback = FInitAnimationCB::CreateLambda([&AnimationRecorder,SampleRate,ExportOptions,SkelMeshComp,AnimSequence]
 	{
 		FAnimationRecordingSettings RecordingSettings;
-		RecordingSettings.SampleRate = SampleRate.AsDecimal();
+		RecordingSettings.SampleFrameRate = SampleRate;
 		RecordingSettings.InterpMode = ERichCurveInterpMode::RCIM_Cubic;
 		RecordingSettings.TangentMode = ERichCurveTangentMode::RCTM_Auto;
 		RecordingSettings.Length = 0;
diff --git a/Engine/Source/Editor/SequenceRecorder/Private/AnimationRecorder.cpp b/Engine/Source/Editor/SequenceRecorder/Private/AnimationRecorder.cpp
--- a/Engine/Source/Editor/SequenceRecorder/Private/AnimationRecorder.cpp
+++ b/Engine/Source/Editor/SequenceRecorder/Private/AnimationRecorder.cpp
@@ -44,7 +44,7 @@ FAnimationRecorder::FAnimationRecorder()
 	, TangentMode(ERichCurveTangentMode::RCTM_Auto)
 	, AnimationSerializer(nullptr)
 {
-	SetSampleRateAndLength(FAnimationRecordingSettings::DefaultSampleRate, FAnimationRecordingSettings::DefaultMaximumLength);
+	SetSampleRateAndLength(FAnimationRecordingSettings::DefaultSampleFrameRate, FAnimationRecordingSettings::DefaultMaximumLength);
 }
 
 FAnimationRecorder::~FAnimationRecorder()
@@ -52,12 +52,12 @@ FAnimationRecorder::~FAnimationRecorder()
 	StopRecord(false);
 }
 
-void FAnimationRecorder::SetSampleRateAndLength(float SampleRateHz, float LengthInSeconds)
+void FAnimationRecorder::SetSampleRateAndLength(FFrameRate SampleFrameRate, float LengthInSeconds)
 {
-	if (SampleRateHz <= 0.f)
+	if (!SampleFrameRate.IsValid())
 	{
 		// invalid rate passed in, fall back to default
-		SampleRateHz = FAnimationRecordingSettings::DefaultSampleRate;
+		SampleFrameRate = FAnimationRecordingSettings::DefaultSampleFrameRate;
 	}
 
 	if (LengthInSeconds <= 0.f)
@@ -66,7 +66,7 @@ void FAnimationRecorder::SetSampleRateAndLength(float SampleRateHz, float Length
 		LengthInSeconds = FAnimationRecordingSettings::UnboundedMaximumLength;
 	}
 
-	IntervalTime = 1.0f / SampleRateHz;
+	RecordingRate = SampleFrameRate;
 	if (LengthInSeconds == FAnimationRecordingSettings::UnboundedMaximumLength)
 	{
 		// invalid length passed in, default to unbounded
@@ -74,7 +74,7 @@ void FAnimationRecorder::SetSampleRateAndLength(float SampleRateHz, float Length
 	}
 	else
 	{
-		MaxFrame = SampleRateHz * LengthInSeconds;
+		MaxFrame = SampleFrameRate.AsFrameNumber(LengthInSeconds);
 	}
 }
 
@@ -96,7 +96,7 @@ bool FAnimationRecorder::SetAnimCompressionScheme(UAnimBoneCompressionSettings*
 }
 
 // Internal. Pops up a dialog to get saved asset path
-static bool PromptUserForAssetDetails(FString& AssetPath, FString& AssetName, float& OutSampleRate, float& OutMaximumDuration)
+static bool PromptUserForAssetDetails(FString& AssetPath, FString& AssetName, FFrameRate& OutSampleRate, float& OutMaximumDuration)
 {
 	TSharedRef<SCreateAnimationDlg> NewAnimDlg = SNew(SCreateAnimationDlg);
 	if (NewAnimDlg->ShowModal() != EAppReturnType::Cancel)
@@ -104,7 +104,7 @@ static bool PromptUserForAssetDetails(FString& AssetPath, FString& AssetName, fl
 		AssetPath = NewAnimDlg->GetFullAssetPath();
 		AssetName = NewAnimDlg->GetAssetName();
 		OutMaximumDuration = NewAnimDlg->GetRecordingParameters()->GetRecordingDurationSeconds();
-		OutSampleRate = NewAnimDlg->GetRecordingParameters()->GetRecordingSampleRate();
+		OutSampleRate = NewAnimDlg->GetRecordingParameters()->GetRecordingFrameRate();
 		return true;
 	}
 
@@ -116,7 +116,7 @@ bool FAnimationRecorder::TriggerRecordAnimation(USkeletalMeshComponent* Componen
 	FString AssetPath;
 	FString AssetName;
 
-	float SampleRate;
+	FFrameRate SampleRate;
 	float MaximumLength;
 
 	if (!Component || !Component->SkeletalMesh || !Component->SkeletalMesh->GetSkeleton())
@@ -156,7 +156,7 @@ bool FAnimationRecorder::TriggerRecordAnimation(USkeletalMeshComponent* Componen
 	if (Parent == nullptr)
 	{
 		// bad or no path passed in, do the popup
-		float SampleRate;
+		FFrameRate SampleRate;
 		float MaximumLength;
 
 		if (PromptUserForAssetDetails(ValidatedAssetPath, ValidatedAssetName, SampleRate, MaximumLength) == false)
@@ -232,7 +232,7 @@ void FAnimationRecorder::GetBoneTransforms(USkeletalMeshComponent* Component, TA
 
 void FAnimationRecorder::StartRecord(USkeletalMeshComponent* Component, UAnimSequence* InAnimationObject)
 {
-	TimePassed = 0.f;
+	TimePassed = 0.0;
 	AnimationObject = InAnimationObject;
 
 	AnimationObject->BoneCompressionSettings = FAnimationUtils::GetDefaultAnimationRecorderBoneCompressionSettings();
@@ -323,14 +323,11 @@ UAnimSequence* FAnimationRecorder::StopRecord(bool bShowMessage)
 	if (AnimationObject)
 	{
 		IAnimationDataController& Controller = AnimationObject->GetController();
-		int32 NumKeys = LastFrame  + 1;
+		int32 NumKeys = LastFrame.Value + 1;
 
 		// can't use TimePassed. That is just total time that has been passed, not necessarily match with frame count
-		Controller.SetPlayLength( (NumKeys>1) ? (NumKeys-1) * IntervalTime : MINIMUM_ANIMATION_LENGTH );
-
-		const float FloatDenominator = 1000.0f;
-		const float Numerator = FloatDenominator / IntervalTime;
-		Controller.SetFrameRate(FFrameRate(Numerator, FloatDenominator));
+		Controller.SetPlayLength( (NumKeys>1) ? RecordingRate.AsSeconds(LastFrame): RecordingRate.AsSeconds(1) );
+		Controller.SetFrameRate(RecordingRate);
 
 		ProcessNotifies();
 
@@ -375,8 +372,8 @@ UAnimSequence* FAnimationRecorder::StopRecord(bool bShowMessage)
 					int32 WriteIndex = 0;
 					for (int32 KeyIndex = 0; KeyIndex < NumKeys; ++KeyIndex)
 					{
-						const float TimeToRecord = KeyIndex * IntervalTime;
-						if (RecordedCurves[KeyIndex].ValidCurveWeights[CurveIndex])
+						const float TimeToRecord = RecordingRate.AsSeconds(KeyIndex);
+						if(RecordedCurves[KeyIndex].ValidCurveWeights[CurveIndex])
 						{
 							float CurCurveValue = RecordedCurves[KeyIndex].CurveWeights[CurveIndex];
 							if (!bSeenThisCurve)
@@ -468,7 +465,7 @@ UAnimSequence* FAnimationRecorder::StopRecord(bool bShowMessage)
 				FText::FromString(AnimationObject->GetName()),
 				FText::AsNumber(AnimationObject->GetDataModel()->GetNumberOfKeys()),
 				FText::AsNumber(AnimationObject->GetPlayLength()),
-				FText::AsNumber(1.f / IntervalTime)
+				RecordingRate.ToPrettyText()
 				);
 					
 			if (GIsEditor)
@@ -513,8 +510,8 @@ void FAnimationRecorder::ProcessRecordedTimes(UAnimSequence* AnimSequence, USkel
 		return;
 	}
 
-	int32 NumFrames = LastFrame  + 1;
-	if (RecordedTimes.Num() != NumFrames)
+	int32 NumKeys = LastFrame.Value + 1;
+	if (RecordedTimes.Num() != NumKeys)
 	{
 		return;
 	}
@@ -530,11 +527,11 @@ void FAnimationRecorder::ProcessRecordedTimes(UAnimSequence* AnimSequence, USkel
 	SubFrames.Reserve(RecordedTimes.Num());
 	Times.Reserve(RecordedTimes.Num());
 
-	for (int32 FrameIndex = 0; FrameIndex < NumFrames; ++FrameIndex)
+	for (int32 KeyIndex = 0; KeyIndex < NumKeys; ++KeyIndex)
 	{
-		const float TimeToRecord = FrameIndex*IntervalTime;
+		const float TimeToRecord = RecordingRate.AsSeconds(KeyIndex);
 
-		FQualifiedFrameTime RecordedTime = RecordedTimes[FrameIndex];
+		FQualifiedFrameTime RecordedTime = RecordedTimes[KeyIndex];
 		FTimecode Timecode = FTimecode::FromFrameNumber(RecordedTime.Time.FrameNumber, RecordedTime.Rate);
 		
 		Hours.Add(Timecode.Hours);
@@ -622,9 +619,9 @@ void FAnimationRecorder::UpdateRecord(USkeletalMeshComponent* Component, float D
 	if (bCheckDeltaTimeAtBeginning)
 	{
 		// in-editor we can get a long frame update because of the modal dialog used to pick paths
-		if (DeltaTime > IntervalTime && (LastFrame == 0 || LastFrame == 1))
+		if (DeltaTime > RecordingRate.AsInterval() && (LastFrame == 0 || LastFrame == 1))
 		{
-			DeltaTime = IntervalTime;
+			DeltaTime = RecordingRate.AsInterval();
 		}
 	}
 
@@ -633,8 +630,8 @@ void FAnimationRecorder::UpdateRecord(USkeletalMeshComponent* Component, float D
 
 	// time passed has been updated
 	// now find what frames we need to update
-	int32 FramesRecorded = LastFrame;
-	int32 FramesToRecord = FPlatformMath::RoundToInt(TimePassed / IntervalTime);
+	int32 FramesRecorded = LastFrame.Value;
+	int32 FramesToRecord = RecordingRate.AsFrameNumber(TimePassed).Value;
 
 	// notifies need to be done regardless of sample rate
 	if (Component->GetAnimInstance())
@@ -662,7 +659,7 @@ void FAnimationRecorder::UpdateRecord(USkeletalMeshComponent* Component, float D
 		{
 			// find what frames we need to record
 			// convert to time
-			const float CurrentTime = (FramesRecorded + 1) * IntervalTime;
+			const float CurrentTime = RecordingRate.AsSeconds(FramesRecorded + 1);
 			float BlendAlpha = (CurrentTime - PreviousTimePassed) / DeltaTime;
 
 			UE_LOG(LogAnimation, Log, TEXT("Current Frame Count : %d, BlendAlpha : %0.2f"), FramesRecorded + 1, BlendAlpha);
@@ -705,7 +702,7 @@ void FAnimationRecorder::UpdateRecord(USkeletalMeshComponent* Component, float D
 	// if we passed MaxFrame, just stop it
 	if (MaxFrame != UnBoundedFrameCount && FramesRecorded >= MaxFrame)
 	{
-		UE_LOG(LogAnimation, Log, TEXT("Animation Recording exceeds the time limited (%d mins). Stopping recording animation... "), (int32)((float)MaxFrame / ((1.0f / IntervalTime) * 60.0f)));
+		UE_LOG(LogAnimation, Log, TEXT("Animation Recording exceeds the time limited (%d seconds). Stopping recording animation... "), RecordingRate.AsSeconds(MaxFrame));
 		FAnimationRecorderManager::Get().StopRecordingAnimation(Component, true);
 		return;
 	}
@@ -778,20 +775,20 @@ bool FAnimationRecorder::Record(USkeletalMeshComponent* Component, FTransform co
 
 				if (bRecordTransforms)
 				{
-					FTransform LocalTransform = SpacesBases[BoneIndex];
-					if (ParentIndex != INDEX_NONE)
-					{
-						LocalTransform.SetToRelativeTransform(SpacesBases[ParentIndex]);
-					}
-					// if record local to world, we'd like to consider component to world to be in root
-					else
+				FTransform LocalTransform = SpacesBases[BoneIndex];
+				if ( ParentIndex != INDEX_NONE )
+				{
+					LocalTransform.SetToRelativeTransform(SpacesBases[ParentIndex]);
+				}
+				// if record local to world, we'd like to consider component to world to be in root
+				else
+				{
+					if (bRecordLocalToWorld)
 					{
-						if (bRecordLocalToWorld)
-						{
-							LocalTransform *= ComponentToWorld;
-						}
+						LocalTransform *= ComponentToWorld;
 					}
-				
+				}
+
 					RawTrack.PosKeys.Add(FVector3f(LocalTransform.GetTranslation()));
 					RawTrack.RotKeys.Add(FQuat4f(LocalTransform.GetRotation()));
 					RawTrack.ScaleKeys.Add(FVector3f(LocalTransform.GetScale3D()));  
@@ -1003,7 +1000,7 @@ void FAnimRecorderInstance::InitInternal(USkeletalMeshComponent* InComponent, co
 {
 	SkelComp = InComponent;
 	Recorder = MakeShareable(new FAnimationRecorder());
-	Recorder->SetSampleRateAndLength(Settings.SampleRate, Settings.Length);
+	Recorder->SetSampleRateAndLength(Settings.SampleFrameRate, Settings.Length);
 	Recorder->bRecordLocalToWorld = Settings.bRecordInWorldSpace;
 	Recorder->InterpMode = Settings.InterpMode;
 	Recorder->TangentMode = Settings.TangentMode;
diff --git a/Engine/Source/Editor/SequenceRecorder/Private/AnimationRecorderParameters.cpp b/Engine/Source/Editor/SequenceRecorder/Private/AnimationRecorderParameters.cpp
--- a/Engine/Source/Editor/SequenceRecorder/Private/AnimationRecorderParameters.cpp
+++ b/Engine/Source/Editor/SequenceRecorder/Private/AnimationRecorderParameters.cpp
@@ -7,7 +7,7 @@ float UAnimationRecordingParameters::GetRecordingDurationSeconds()
 	return bEndAfterDuration ? MaximumDurationSeconds : FAnimationRecordingSettings::UnboundedMaximumLength;
 }
 
-float UAnimationRecordingParameters::GetRecordingSampleRate()
+const FFrameRate& UAnimationRecordingParameters::GetRecordingFrameRate()
 {
-	return SampleRate;
-}
\ No newline at end of file
+	return SampleFrameRate;
+}
diff --git a/Engine/Source/Editor/SequenceRecorder/Private/AnimationRecorderParameters.h b/Engine/Source/Editor/SequenceRecorder/Private/AnimationRecorderParameters.h
--- a/Engine/Source/Editor/SequenceRecorder/Private/AnimationRecorderParameters.h
+++ b/Engine/Source/Editor/SequenceRecorder/Private/AnimationRecorderParameters.h
@@ -15,17 +15,34 @@ class UAnimationRecordingParameters : public UObject
 
 public:
 	float GetRecordingDurationSeconds();
-	float GetRecordingSampleRate();
+	const FFrameRate& GetRecordingFrameRate();
 
-	/** Sample rate of the recorded animation (in Hz) */
-	UPROPERTY(EditAnywhere, config, Category = Options, meta = (UIMin = "1.0", ClampMin = "1.0", UIMax = "60.0", ClampMax = "60.0"))
-	float SampleRate = FAnimationRecordingSettings::DefaultSampleRate;
+	virtual void PostInitProperties() override
+	{
+		Super::PostInitProperties();
+		PRAGMA_DISABLE_DEPRECATION_WARNINGS
+		if (SampleRate > 0.f)
+		{
+			FString FrameRateFormat = FString::Printf(TEXT("1/%f"), SampleRate);
+			TryParseString(SampleFrameRate, *FrameRateFormat);
+			SampleRate = 0.f;
+		}
+		PRAGMA_ENABLE_DEPRECATION_WARNINGS
+	}
+
+	/** Sample frame-rate of the recorded animation  */
+	UPROPERTY(EditAnywhere, config, Category = Options)
+	FFrameRate SampleFrameRate = FAnimationRecordingSettings::DefaultSampleFrameRate;
 
 	/** If enabled, this animation recording will automatically end after a set amount of time */
 	UPROPERTY(EditAnywhere, config, Category = Options)
 	bool bEndAfterDuration = false;
 
 	/** The maximum duration of this animation recording */
-	UPROPERTY(EditAnywhere, config, Category = Options, meta = (EditCondition = "bEndAfterDuration", UIMin = "5.0", ClampMin = "5.0"))
+	UPROPERTY(EditAnywhere, config, Category = Options, meta = (EditCondition = "bEndAfterDuration", UIMin = "0.03333", ClampMin = "0.03333"))
 	float MaximumDurationSeconds = FAnimationRecordingSettings::DefaultMaximumLength;
+
+	UE_DEPRECATED(5.0, "SampleRate has been deprecated and replaced with SampleFrameRate")
+	UPROPERTY(config)
+	float SampleRate = 0.f;
 };
diff --git a/Engine/Source/Editor/SequenceRecorder/Private/SequenceRecorder.cpp b/Engine/Source/Editor/SequenceRecorder/Private/SequenceRecorder.cpp
--- a/Engine/Source/Editor/SequenceRecorder/Private/SequenceRecorder.cpp
+++ b/Engine/Source/Editor/SequenceRecorder/Private/SequenceRecorder.cpp
@@ -544,7 +544,7 @@ void FSequenceRecorder::DrawDebug(UCanvas* InCanvas, APlayerController* InPlayer
 		TimeAccumulator -= Minutes * 60.0f;
 		float Seconds = FMath::FloorToFloat(TimeAccumulator);
 		TimeAccumulator -= Seconds;
-		float Frames = FMath::FloorToFloat(TimeAccumulator * GetDefault<USequenceRecorderSettings>()->DefaultAnimationSettings.SampleRate);
+		float Frames = FMath::FloorToFloat(GetDefault<USequenceRecorderSettings>()->DefaultAnimationSettings.SampleFrameRate.AsFrameTime(TimeAccumulator).AsDecimal());
 
 		FNumberFormattingOptions Options;
 		Options.MinimumIntegralDigits = 2;
diff --git a/Engine/Source/Editor/SequenceRecorder/Private/SequenceRecorderModule.cpp b/Engine/Source/Editor/SequenceRecorder/Private/SequenceRecorderModule.cpp
--- a/Engine/Source/Editor/SequenceRecorder/Private/SequenceRecorderModule.cpp
+++ b/Engine/Source/Editor/SequenceRecorder/Private/SequenceRecorderModule.cpp
@@ -57,10 +57,32 @@ static TAutoConsoleVariable<float> CVarDefaultRecordedAnimLength(
 	TEXT("Sets default animation length for the animation recorder system."),
 	ECVF_Default);
 
-static TAutoConsoleVariable<float> CVarAnimRecorderSampleRate(
+static FAutoConsoleCommand CVarAnimRecorderSampleFrameRate(
 	TEXT("AnimRecorder.SampleRate"),
-	FAnimationRecordingSettings::DefaultSampleRate,
-	TEXT("Sets the sample rate for the animation recorder system"),
+	TEXT("Argument: valid Frame Rate format\n")
+	TEXT("Sets the sample frame-rate for the animation recorder system"),
+	FConsoleCommandWithArgsDelegate::CreateLambda([](const TArray<FString>& Args)
+	{
+		bool bValidFrameRate = false;
+		if (Args.Num() == 1)
+		{
+			TValueOrError<FFrameRate, FExpressionError> ParseResult = ParseFrameRate(*Args[0]);
+
+			if (ParseResult.IsValid())
+			{
+				GetMutableDefault<USequenceRecorderSettings>()->DefaultAnimationSettings.SampleFrameRate = ParseResult.GetValue();
+				UE_LOG(LogEngine, Display, TEXT("Sequence recorder default sample frame rate set to: %s"), *(GetMutableDefault<USequenceRecorderSettings>()->DefaultAnimationSettings.SampleFrameRate.ToPrettyText().ToString()));
+			}
+			else
+			{
+				UE_LOG(LogAnimation, Warning, TEXT("Invalid Frame Rate format provided for AnimRecorder.SampleRate:\n%s"), *ParseResult.GetError().Text.ToString());
+			}
+		}
+		else
+		{
+			UE_LOG(LogAnimation, Warning, TEXT("Missing Frame Rate argument for AnimRecorder.SampleRate"));
+		}
+	}),
 	ECVF_Default);
 
 static TAutoConsoleVariable<int32> CVarAnimRecorderWorldSpace(
@@ -110,11 +132,6 @@ class FSequenceRecorderModule : public ISequenceRecorder, private FSelfRegisteri
 			GetMutableDefault<USequenceRecorderSettings>()->DefaultAnimationSettings.Length = CVarDefaultRecordedAnimLength.GetValueOnGameThread();
 		}));
 
-		CVarAnimRecorderSampleRate.AsVariable()->SetOnChangedCallback(FConsoleVariableDelegate::CreateLambda([](IConsoleVariable* Variable)
-		{
-			GetMutableDefault<USequenceRecorderSettings>()->DefaultAnimationSettings.SampleRate = CVarAnimRecorderSampleRate.GetValueOnGameThread();
-		}));
-
 		CVarAnimRecorderWorldSpace.AsVariable()->SetOnChangedCallback(FConsoleVariableDelegate::CreateLambda([](IConsoleVariable* Variable)
 		{
 			GetMutableDefault<USequenceRecorderSettings>()->DefaultAnimationSettings.bRecordInWorldSpace = (CVarAnimRecorderWorldSpace.GetValueOnGameThread() != 0);
diff --git a/Engine/Source/Editor/SequenceRecorder/Private/SequenceRecorderUtils.cpp b/Engine/Source/Editor/SequenceRecorder/Private/SequenceRecorderUtils.cpp
--- a/Engine/Source/Editor/SequenceRecorder/Private/SequenceRecorderUtils.cpp
+++ b/Engine/Source/Editor/SequenceRecorder/Private/SequenceRecorderUtils.cpp
@@ -62,8 +62,8 @@ bool RecordSingleNodeInstanceToAnimation(USkeletalMeshComponent* PreviewComponen
 		FAnimationRecordingSettings Setting;
 		AnimRecorder.Init(PreviewComponent, NewAsset,nullptr, Setting);
 		float Length = SingleNodeInstance->GetLength();
-		const float DefaultSampleRate = (Setting.SampleRate > 0.f) ? Setting.SampleRate : DEFAULT_SAMPLERATE;
-		const float Interval = 1.f / DefaultSampleRate;
+		const FFrameRate DefaultSampleRate = Setting.SampleFrameRate.IsValid() ? Setting.SampleFrameRate : FAnimationRecordingSettings::DefaultSampleFrameRate;
+		const float Interval = DefaultSampleRate.AsInterval();
 		float Time = 0.f;
 		for (; Time < Length; Time += Interval)
 		{
diff --git a/Engine/Source/Editor/SequenceRecorder/Public/AnimationRecorder.h b/Engine/Source/Editor/SequenceRecorder/Public/AnimationRecorder.h
--- a/Engine/Source/Editor/SequenceRecorder/Public/AnimationRecorder.h
+++ b/Engine/Source/Editor/SequenceRecorder/Public/AnimationRecorder.h
@@ -31,10 +31,10 @@ private:
 	static const int32 UnBoundedFrameCount = -1;
 
 private:
-	float IntervalTime;
-	int32 MaxFrame;
-	int32 LastFrame;
-	float TimePassed;
+	FFrameRate RecordingRate;
+	FFrameNumber MaxFrame;
+	FFrameNumber LastFrame;
+	double TimePassed;
 	UAnimSequence* AnimationObject;
 	TArray<FTransform> PreviousSpacesBases;
 	FBlendedHeapCurve PreviousAnimCurves;
@@ -81,10 +81,10 @@ public:
 	void UpdateRecord(USkeletalMeshComponent* Component, float DeltaTime);
 	UAnimSequence* GetAnimationObject() const { return AnimationObject; }
 	bool InRecording() const { return AnimationObject != nullptr; }
-	float GetTimeRecorded() const { return TimePassed; }
+	double GetTimeRecorded() const { return TimePassed; }
 
 	/** Sets a new sample rate & max length for this recorder. Don't call while recording. */
-	void SetSampleRateAndLength(float SampleRateHz, float LengthInMinutes);
+	void SetSampleRateAndLength(FFrameRate SampleFrameRate, float LengthInSeconds);
 
 	bool SetAnimCompressionScheme(UAnimBoneCompressionSettings* Settings);
 
diff --git a/Engine/Source/Editor/SequenceRecorderSections/Private/MovieSceneSpawnSectionRecorder.cpp b/Engine/Source/Editor/SequenceRecorderSections/Private/MovieSceneSpawnSectionRecorder.cpp
--- a/Engine/Source/Editor/SequenceRecorderSections/Private/MovieSceneSpawnSectionRecorder.cpp
+++ b/Engine/Source/Editor/SequenceRecorderSections/Private/MovieSceneSpawnSectionRecorder.cpp
@@ -86,7 +86,7 @@ void FMovieSceneSpawnSectionRecorder::FinalizeSection(float CurrentTime)
 		FMovieSceneBoolChannel* Channel = MovieSceneSection->GetChannelProxy().GetChannel<FMovieSceneBoolChannel>(0);
 		if (ensure(Channel))
 		{
-			double       OneFrameInterval = 1.0/GetDefault<USequenceRecorderSettings>()->DefaultAnimationSettings.SampleRate;
+			double       OneFrameInterval = GetDefault<USequenceRecorderSettings>()->DefaultAnimationSettings.SampleFrameRate.AsInterval();
 
 			FFrameRate   TickResolution   = MovieSceneSection->GetTypedOuter<UMovieScene>()->GetTickResolution();
 			FFrameNumber StartTime        = MovieSceneSection->GetExclusiveEndFrame() - (OneFrameInterval * TickResolution).CeilToFrame();
diff --git a/Engine/Source/Runtime/Engine/Classes/Animation/AnimationRecordingSettings.h b/Engine/Source/Runtime/Engine/Classes/Animation/AnimationRecordingSettings.h
--- a/Engine/Source/Runtime/Engine/Classes/Animation/AnimationRecordingSettings.h
+++ b/Engine/Source/Runtime/Engine/Classes/Animation/AnimationRecordingSettings.h
@@ -13,8 +13,8 @@ struct ENGINE_API FAnimationRecordingSettings
 {
 	GENERATED_BODY()
 
-	/** 30Hz default sample rate */
-	static const float DefaultSampleRate;
+	/** 30Hz default sample frame rate */
+	static const FFrameRate DefaultSampleFrameRate;
 
 	/** 1 minute default length */
 	static const float DefaultMaximumLength;
@@ -26,7 +26,7 @@ struct ENGINE_API FAnimationRecordingSettings
 		: bRecordInWorldSpace(true)
 		, bRemoveRootAnimation(true)
 		, bAutoSaveAsset(false)
-		, SampleRate((float)DefaultSampleRate)
+		, SampleFrameRate(DefaultSampleFrameRate)
 		, Length((float)DefaultMaximumLength)
 		, InterpMode(ERichCurveInterpMode::RCIM_Linear)
 		, TangentMode(ERichCurveTangentMode::RCTM_Auto)
@@ -49,9 +49,9 @@ struct ENGINE_API FAnimationRecordingSettings
 	UPROPERTY(EditAnywhere, Category = "Settings")
 	bool bAutoSaveAsset;
 
-	/** Sample rate of the recorded animation (in Hz) */
+	/** Sample rate of the recorded animation */
 	UPROPERTY(EditAnywhere, Category = "Settings")
-	float SampleRate;
+	FFrameRate SampleFrameRate;
 
 	/** Maximum length of the animation recorded (in seconds). If zero the animation will keep on recording until stopped. */
 	UPROPERTY(EditAnywhere, Category = "Settings")
diff --git a/Engine/Source/Runtime/Engine/Private/Animation/AnimationRecordingSettings.cpp b/Engine/Source/Runtime/Engine/Private/Animation/AnimationRecordingSettings.cpp
--- a/Engine/Source/Runtime/Engine/Private/Animation/AnimationRecordingSettings.cpp
+++ b/Engine/Source/Runtime/Engine/Private/Animation/AnimationRecordingSettings.cpp
@@ -4,7 +4,7 @@
 #include "Animation/AnimTypes.h"
 
 /** 30Hz default sample rate */
-const float FAnimationRecordingSettings::DefaultSampleRate = DEFAULT_SAMPLERATE;
+const FFrameRate FAnimationRecordingSettings::DefaultSampleFrameRate = FFrameRate(30, 1);
 
 /** 1 minute default length */
 const float FAnimationRecordingSettings::DefaultMaximumLength = 1.0f * 60.0f;
