Sequencer: Control Rig: Bakding crash due to a level getting unloaded while backing need to just exit if the number of spaces changes on the skeleton.
#trivial

diff --git a/Engine/Source/Editor/SequenceRecorder/Private/AnimationRecorder.cpp b/Engine/Source/Editor/SequenceRecorder/Private/AnimationRecorder.cpp
--- a/Engine/Source/Editor/SequenceRecorder/Private/AnimationRecorder.cpp
+++ b/Engine/Source/Editor/SequenceRecorder/Private/AnimationRecorder.cpp
@@ -646,7 +646,12 @@ void FAnimationRecorder::UpdateRecord(USkeletalMeshComponent* Component, float D
 	{
 		const FBlendedHeapCurve& AnimCurves = Component->GetAnimationCurves();
 
-		check(SpaceBases.Num() == PreviousSpacesBases.Num());
+		if (SpaceBases.Num() != PreviousSpacesBases.Num())
+		{
+			UE_LOG(LogAnimation, Log, TEXT("Current Num of Spaces %d don't match with the previous number %d so we are stopping recording"), SpaceBases.Num(),PreviousSpacesBases.Num());
+			StopRecord(true);
+			return;
+		}
 
 		TArray<FTransform> BlendedSpaceBases;
 		BlendedSpaceBases.AddZeroed(SpaceBases.Num());
