Sequencer: Take Recorder: Set Retarget Source Asset when take/animation recording if no retarget source from mesh is specified.

diff --git a/Engine/Source/Editor/SequenceRecorder/Private/AnimationRecorder.cpp b/Engine/Source/Editor/SequenceRecorder/Private/AnimationRecorder.cpp
--- a/Engine/Source/Editor/SequenceRecorder/Private/AnimationRecorder.cpp
+++ b/Engine/Source/Editor/SequenceRecorder/Private/AnimationRecorder.cpp
@@ -278,6 +278,10 @@ void FAnimationRecorder::StartRecord(USkeletalMeshComponent* Component, UAnimSeq
 	}
 
 	AnimationObject->RetargetSource = Component->SkeletalMesh ? AnimSkeleton->GetRetargetSourceForMesh(Component->SkeletalMesh) : NAME_None;
+	if (AnimationObject->RetargetSource == NAME_None)
+	{
+		AnimationObject->RetargetSourceAsset = Component->SkeletalMesh;
+	}
 
 	// record the first frame
 	Record(Component, PreviousComponentToWorld, PreviousSpacesBases, PreviousAnimCurves,  0);
