Bake Animation Sequence when baking some curves
#fix ensure that WriteIndex is used to populate curve-key data, which prevents frame-gaps from causing invalid keys

diff --git a/Engine/Source/Editor/SequenceRecorder/Private/AnimationRecorder.cpp b/Engine/Source/Editor/SequenceRecorder/Private/AnimationRecorder.cpp
--- a/Engine/Source/Editor/SequenceRecorder/Private/AnimationRecorder.cpp
+++ b/Engine/Source/Editor/SequenceRecorder/Private/AnimationRecorder.cpp
@@ -397,8 +397,8 @@ UAnimSequence* FAnimationRecorder::StopRecord(bool bShowMessage)
 
 							if (FloatCurveData)
 							{
-								TimesToRecord[KeyIndex] = TimeToRecord;
-								ValuesToRecord[KeyIndex] = CurCurveValue;
+								TimesToRecord[WriteIndex] = TimeToRecord;
+								ValuesToRecord[WriteIndex] = CurCurveValue;
 
 								++WriteIndex;
 							}
