Take Recorder: Add option to specify the bone to record timecodes onto

FTimecodeBoneMethod specifies the mode (ie. All, Root, UserDefined) and the bone name to assign to.
If the bone name doesn't exist, it will fall back to the root.

diff --git a/Engine/Plugins/VirtualProduction/Takes/Source/TakeTrackRecorders/Private/MovieSceneAnimationTrackRecorder.cpp b/Engine/Plugins/VirtualProduction/Takes/Source/TakeTrackRecorders/Private/MovieSceneAnimationTrackRecorder.cpp
--- a/Engine/Plugins/VirtualProduction/Takes/Source/TakeTrackRecorders/Private/MovieSceneAnimationTrackRecorder.cpp
+++ b/Engine/Plugins/VirtualProduction/Takes/Source/TakeTrackRecorders/Private/MovieSceneAnimationTrackRecorder.cpp
@@ -327,7 +327,9 @@ void UMovieSceneAnimationTrackRecorder::RemoveRootMotion()
 
 void UMovieSceneAnimationTrackRecorder::ProcessRecordedTimes(const FString& HoursName, const FString& MinutesName, const FString& SecondsName, const FString& FramesName, const FString& SubFramesName, const FString& SlateName, const FString& Slate)
 {
-	AnimationRecorder.ProcessRecordedTimes(AnimSequence.Get(), SkeletalMeshComponent.Get(), HoursName, MinutesName, SecondsName, FramesName, SubFramesName, SlateName, Slate);
+	UMovieSceneAnimationTrackRecorderSettings* AnimSettings = CastChecked<UMovieSceneAnimationTrackRecorderSettings>(Settings.Get());
+
+	AnimationRecorder.ProcessRecordedTimes(AnimSequence.Get(), SkeletalMeshComponent.Get(), HoursName, MinutesName, SecondsName, FramesName, SubFramesName, SlateName, Slate, AnimSettings->TimecodeBoneMethod);
 }
 
 bool UMovieSceneAnimationTrackRecorder::LoadRecordedFile(const FString& FileName, UMovieScene *InMovieScene, TMap<FGuid, AActor*>& ActorGuidToActorMap,  TFunction<void()> InCompletionCallback)
diff --git a/Engine/Plugins/VirtualProduction/Takes/Source/TakeTrackRecorders/Public/TrackRecorders/MovieSceneAnimationTrackRecorderSettings.h b/Engine/Plugins/VirtualProduction/Takes/Source/TakeTrackRecorders/Public/TrackRecorders/MovieSceneAnimationTrackRecorderSettings.h
--- a/Engine/Plugins/VirtualProduction/Takes/Source/TakeTrackRecorders/Public/TrackRecorders/MovieSceneAnimationTrackRecorderSettings.h
+++ b/Engine/Plugins/VirtualProduction/Takes/Source/TakeTrackRecorders/Public/TrackRecorders/MovieSceneAnimationTrackRecorderSettings.h
@@ -6,6 +6,7 @@
 #include "Misc/FrameRate.h"
 #include "MovieSceneTrackRecorderSettings.h"
 #include "Curves/RichCurve.h"
+#include "AnimationRecorder.h"
 #include "MovieSceneAnimationTrackRecorderSettings.generated.h"
 
 UCLASS(Abstract, BlueprintType, config=EditorSettings, DisplayName="Animation Recorder")
@@ -20,6 +21,7 @@ public:
 	, AnimationSubDirectory(TEXT("Animation"))
 	, bRemoveRootAnimation(true)
 	{
+		TimecodeBoneMethod.BoneMode = ETimecodeBoneMode::Root;
 	}
 
 	virtual void PostEditChangeProperty(FPropertyChangedEvent& PropertyChangedEvent) override;
@@ -48,7 +50,9 @@ public:
 	UPROPERTY(config, EditAnywhere, BlueprintReadWrite, Category = "Animation Recorder Settings")
 	bool bRemoveRootAnimation;
 
-
+	/** The method to record timecode values onto bones */
+	UPROPERTY(config, EditAnywhere, BlueprintReadWrite, Category = "Animation Recorder Settings", meta = (ShowOnlyInnerProperties))
+	FTimecodeBoneMethod TimecodeBoneMethod;
 };
 
 UCLASS(BlueprintType, config = EditorSettings, DisplayName = "Animation Recorder Settings")
diff --git a/Engine/Source/Editor/SequenceRecorder/Private/AnimationRecorder.cpp b/Engine/Source/Editor/SequenceRecorder/Private/AnimationRecorder.cpp
--- a/Engine/Source/Editor/SequenceRecorder/Private/AnimationRecorder.cpp
+++ b/Engine/Source/Editor/SequenceRecorder/Private/AnimationRecorder.cpp
@@ -507,7 +507,7 @@ UAnimSequence* FAnimationRecorder::StopRecord(bool bShowMessage)
 	return nullptr;
 }
 
-void FAnimationRecorder::ProcessRecordedTimes(UAnimSequence* AnimSequence, USkeletalMeshComponent* SkeletalMeshComponent, const FString& HoursName, const FString& MinutesName, const FString& SecondsName, const FString& FramesName, const FString& SubFramesName, const FString& SlateName, const FString& Slate)
+void FAnimationRecorder::ProcessRecordedTimes(UAnimSequence* AnimSequence, USkeletalMeshComponent* SkeletalMeshComponent, const FString& HoursName, const FString& MinutesName, const FString& SecondsName, const FString& FramesName, const FString& SubFramesName, const FString& SlateName, const FString& Slate, const FTimecodeBoneMethod& TimecodeBoneMethod)
 {
 	if (!AnimSequence || !SkeletalMeshComponent)
 	{
@@ -576,6 +576,30 @@ void FAnimationRecorder::ProcessRecordedTimes(UAnimSequence* AnimSequence, USkel
 	IAnimationDataController& Controller = AnimSequence->GetController();
 	IAnimationDataController::FScopedBracket(Controller, LOCTEXT("AddTimeCodeAttributesBracket", "Adding Time Code attributes"));
 
+	// If the user defined bone doesn't exist, fallback to writing timecodes to the root
+	bool bHasUserDefinedBone = false;
+	if (TimecodeBoneMethod.BoneMode == ETimecodeBoneMode::UserDefined)
+	{
+		for (int32 BoneIndex = 0; BoneIndex < SpaceBases->Num(); ++BoneIndex)
+		{
+			const int32 BoneTreeIndex = AnimSkeleton->GetSkeletonBoneIndexFromMeshBoneIndex(SkeletalMeshComponent->MasterPoseComponent != nullptr ? SkeletalMeshComponent->MasterPoseComponent->SkeletalMesh : SkeletalMeshComponent->SkeletalMesh, BoneIndex);
+			if (BoneTreeIndex != INDEX_NONE)
+			{
+				FName BoneTreeName = AnimSkeleton->GetReferenceSkeleton().GetBoneName(BoneTreeIndex);
+				if (BoneTreeName == TimecodeBoneMethod.BoneName)
+				{
+					bHasUserDefinedBone = true;
+					break;
+				}
+			}
+		}
+
+		if (!bHasUserDefinedBone)
+		{
+			UE_LOG(LogAnimation, Warning, TEXT("User defined bone name: %s not found. Falling back to assigning timecodes to root bone"), *TimecodeBoneMethod.BoneName.ToString());
+		}
+	}
+
 	for (int32 BoneIndex = 0; BoneIndex < SpaceBases->Num(); ++BoneIndex)
 	{
 		// verify if this bone exists in skeleton
@@ -585,6 +609,17 @@ void FAnimationRecorder::ProcessRecordedTimes(UAnimSequence* AnimSequence, USkel
 			// add tracks for the bone existing
 			FName BoneTreeName = AnimSkeleton->GetReferenceSkeleton().GetBoneName(BoneTreeIndex);
 			
+			const bool bUseThisBone = 
+				TimecodeBoneMethod.BoneMode == ETimecodeBoneMode::All ||
+				(TimecodeBoneMethod.BoneMode == ETimecodeBoneMode::Root && BoneIndex == 0) ||
+				(TimecodeBoneMethod.BoneMode == ETimecodeBoneMode::UserDefined && BoneTreeName == TimecodeBoneMethod.BoneName) ||
+				(bHasUserDefinedBone == false && BoneIndex == 0);
+
+			if (!bUseThisBone)
+			{
+				continue;
+			}
+
 			UE::Anim::AddTypedCustomAttribute<FIntegerAnimationAttribute, int32>(FName(*HoursName), BoneTreeName, AnimSequence, MakeArrayView(Times), MakeArrayView(Hours));
 			UE::Anim::AddTypedCustomAttribute<FIntegerAnimationAttribute, int32>(FName(*MinutesName), BoneTreeName, AnimSequence, MakeArrayView(Times), MakeArrayView(Minutes));
 			UE::Anim::AddTypedCustomAttribute<FIntegerAnimationAttribute, int32>(FName(*SecondsName), BoneTreeName, AnimSequence, MakeArrayView(Times), MakeArrayView(Seconds));
@@ -1098,11 +1133,11 @@ void FAnimRecorderInstance::FinishRecording(bool bShowMessage)
 	}
 }
 
-void FAnimRecorderInstance::ProcessRecordedTimes(UAnimSequence* AnimSequence, USkeletalMeshComponent* SkeletalMeshComponent, const FString& HoursName, const FString& MinutesName, const FString& SecondsName, const FString& FramesName, const FString& SubFramesName, const FString& SlateName, const FString& Slate)
+void FAnimRecorderInstance::ProcessRecordedTimes(UAnimSequence* AnimSequence, USkeletalMeshComponent* SkeletalMeshComponent, const FString& HoursName, const FString& MinutesName, const FString& SecondsName, const FString& FramesName, const FString& SubFramesName, const FString& SlateName, const FString& Slate, const FTimecodeBoneMethod& TimecodeBoneMethod)
 {
 	if (Recorder.IsValid())
 	{
-		Recorder->ProcessRecordedTimes(AnimSequence, SkeletalMeshComponent, HoursName, MinutesName, SecondsName, FramesName, SubFramesName, SlateName, Slate);
+		Recorder->ProcessRecordedTimes(AnimSequence, SkeletalMeshComponent, HoursName, MinutesName, SecondsName, FramesName, SubFramesName, SlateName, Slate, TimecodeBoneMethod);
 	}
 }
 
diff --git a/Engine/Source/Editor/SequenceRecorder/Public/AnimationRecorder.h b/Engine/Source/Editor/SequenceRecorder/Public/AnimationRecorder.h
--- a/Engine/Source/Editor/SequenceRecorder/Public/AnimationRecorder.h
+++ b/Engine/Source/Editor/SequenceRecorder/Public/AnimationRecorder.h
@@ -11,6 +11,7 @@
 #include "Animation/AnimNotifyQueue.h"
 #include "Serializers/MovieSceneAnimationSerialization.h"
 #include "Misc/QualifiedFrameTime.h"
+#include "AnimationRecorder.generated.h"
 
 class UAnimBoneCompressionSettings;
 class UAnimNotify;
@@ -20,6 +21,32 @@ class USkeletalMeshComponent;
 
 DECLARE_LOG_CATEGORY_EXTERN(AnimationSerialization, Verbose, All);
 
+UENUM(BlueprintType)
+enum class ETimecodeBoneMode : uint8
+{
+	All,
+	Root,
+	UserDefined,
+	MAX UMETA(Hidden)
+};
+
+USTRUCT(BlueprintType)
+struct FTimecodeBoneMethod
+{
+	GENERATED_USTRUCT_BODY()
+
+	/** Default constructor, initializing with default values */
+	FTimecodeBoneMethod() { }
+
+	/** The timecode bone mode */
+	UPROPERTY(EditAnywhere, BlueprintReadWrite, Category="Timecode")
+	ETimecodeBoneMode BoneMode;
+
+	/** Name of the bone to assign timecode values to */
+	UPROPERTY(EditAnywhere, BlueprintReadWrite, Category="Timecode")
+	FName BoneName;
+};
+
 //////////////////////////////////////////////////////////////////////////
 // FAnimationRecorder
 
@@ -90,7 +117,7 @@ public:
 
 	const FTransform& GetInitialRootTransform() const { return InitialRootTransform; }
 
-	void ProcessRecordedTimes(UAnimSequence* AnimSequence, USkeletalMeshComponent* SkeletalMeshComponent, const FString& HoursName, const FString& MinutesName, const FString& SecondsName, const FString& FramesName, const FString& SubFramesName, const FString& SlateName, const FString& Slate);
+	void ProcessRecordedTimes(UAnimSequence* AnimSequence, USkeletalMeshComponent* SkeletalMeshComponent, const FString& HoursName, const FString& MinutesName, const FString& SecondsName, const FString& FramesName, const FString& SubFramesName, const FString& SlateName, const FString& Slate, const FTimecodeBoneMethod& TimecodeBoneMethod);
 
 	/** If true, it will record root to include LocalToWorld */
 	uint8 bRecordLocalToWorld :1;
@@ -160,7 +187,7 @@ public:
 	bool BeginRecording();
 	void Update(float DeltaTime);
 	void FinishRecording(bool bShowMessage = true);
-	void ProcessRecordedTimes(UAnimSequence* AnimSequence, USkeletalMeshComponent* SkeletalMeshComponent, const FString& HoursName, const FString& MinutesName, const FString& SecondsName, const FString& FramesName, const FString& SubFramesName, const FString& SlateName, const FString& Slate);
+	void ProcessRecordedTimes(UAnimSequence* AnimSequence, USkeletalMeshComponent* SkeletalMeshComponent, const FString& HoursName, const FString& MinutesName, const FString& SecondsName, const FString& FramesName, const FString& SubFramesName, const FString& SlateName, const FString& Slate, const FTimecodeBoneMethod& TimecodeBoneMethod);
 
 private:
 	void InitInternal(USkeletalMeshComponent* InComponent, const FAnimationRecordingSettings& Settings, FAnimationSerializer *InAnimationSerializer = nullptr);
