Take Recorder:  Fix initialization

diff --git a/Engine/Source/Editor/SequenceRecorder/Public/AnimationRecorder.h b/Engine/Source/Editor/SequenceRecorder/Public/AnimationRecorder.h
--- a/Engine/Source/Editor/SequenceRecorder/Public/AnimationRecorder.h
+++ b/Engine/Source/Editor/SequenceRecorder/Public/AnimationRecorder.h
@@ -36,7 +36,7 @@ struct FTimecodeBoneMethod
 	GENERATED_USTRUCT_BODY()
 
 	/** Default constructor, initializing with default values */
-	FTimecodeBoneMethod() { }
+	FTimecodeBoneMethod() : BoneMode(ETimecodeBoneMode::Root) { }
 
 	/** The timecode bone mode */
 	UPROPERTY(EditAnywhere, BlueprintReadWrite, Category="Timecode")
