Take Recorder: Fix include

diff --git a/Engine/Source/Editor/SequenceRecorder/Private/TimecodeBoneMethodCustomization.cpp b/Engine/Source/Editor/SequenceRecorder/Private/TimecodeBoneMethodCustomization.cpp
--- a/Engine/Source/Editor/SequenceRecorder/Private/TimecodeBoneMethodCustomization.cpp
+++ b/Engine/Source/Editor/SequenceRecorder/Private/TimecodeBoneMethodCustomization.cpp
@@ -10,6 +10,7 @@
 #include "IDetailChildrenBuilder.h"
 #include "LevelEditor.h"
 #include "Misc/Attribute.h"
+#include "Modules/ModuleManager.h"
 #include "PropertyHandle.h"
 #include "PropertyCustomizationHelpers.h"
 #include "ScopedTransaction.h"
