Take Recorder: Added the ability to filter out bones/curves that are being recorded.

AnimationRecorder now has IncludeAnimationNames and ExcludeAnimationNames lists. Bones and curves are tested against this list to determine whether they should be recorded.

diff --git a/Engine/Plugins/VirtualProduction/Takes/Source/TakeRecorderSources/Private/TakeRecorderActorSource.cpp b/Engine/Plugins/VirtualProduction/Takes/Source/TakeRecorderSources/Private/TakeRecorderActorSource.cpp
--- a/Engine/Plugins/VirtualProduction/Takes/Source/TakeRecorderSources/Private/TakeRecorderActorSource.cpp
+++ b/Engine/Plugins/VirtualProduction/Takes/Source/TakeRecorderSources/Private/TakeRecorderActorSource.cpp
@@ -1656,6 +1656,8 @@ FTrackRecorderSettings UTakeRecorderActorSource::GetTrackRecorderSettings() cons
 	TrackRecorderSettings.ReduceKeysTolerance = Parameters.User.ReduceKeysTolerance;
 
 	TrackRecorderSettings.DefaultTracks = Parameters.Project.DefaultTracks;
+	TrackRecorderSettings.IncludeAnimationNames = IncludeAnimationNames;
+	TrackRecorderSettings.ExcludeAnimationNames = ExcludeAnimationNames;
 
 	return TrackRecorderSettings;
 }
diff --git a/Engine/Plugins/VirtualProduction/Takes/Source/TakeRecorderSources/Private/TakeRecorderActorSource.h b/Engine/Plugins/VirtualProduction/Takes/Source/TakeRecorderSources/Private/TakeRecorderActorSource.h
--- a/Engine/Plugins/VirtualProduction/Takes/Source/TakeRecorderSources/Private/TakeRecorderActorSource.h
+++ b/Engine/Plugins/VirtualProduction/Takes/Source/TakeRecorderSources/Private/TakeRecorderActorSource.h
@@ -68,6 +68,14 @@ public:
 	UPROPERTY(EditAnywhere, Instanced, BlueprintReadWrite, meta=(ShowInnerProperties), Category = "Actor Source")
 	TObjectPtr<UActorRecorderPropertyMap> RecordedProperties;
 
+	/** Include only the animation bones/curves that match this list */
+	UPROPERTY(EditAnywhere, BlueprintReadWrite, Category = "Actor Source")
+	TArray<FString> IncludeAnimationNames;
+
+	/** Exclude all animation bones/curves that match this list */
+	UPROPERTY(EditAnywhere, BlueprintReadWrite, Category = "Actor Source")
+	TArray<FString> ExcludeAnimationNames;
+
 	/** The level sequence that this source is being recorded into. Set during PreRecording, null after PostRecording. */
 	UPROPERTY()
 	TObjectPtr<ULevelSequence> TargetLevelSequence;
diff --git a/Engine/Plugins/VirtualProduction/Takes/Source/TakeTrackRecorders/Private/MovieSceneAnimationTrackRecorder.cpp b/Engine/Plugins/VirtualProduction/Takes/Source/TakeTrackRecorders/Private/MovieSceneAnimationTrackRecorder.cpp
--- a/Engine/Plugins/VirtualProduction/Takes/Source/TakeTrackRecorders/Private/MovieSceneAnimationTrackRecorder.cpp
+++ b/Engine/Plugins/VirtualProduction/Takes/Source/TakeTrackRecorders/Private/MovieSceneAnimationTrackRecorder.cpp
@@ -272,6 +272,8 @@ void UMovieSceneAnimationTrackRecorder::RecordSampleImpl(const FQualifiedFrameTi
 		//If not removing root we also don't record in world space ( not totally sure if it matters but matching up with Sequence Recorder)
 		bool bRecordInWorldSpace = bRootWasRemoved == false ? false : true;
 
+		FTrackRecorderSettings TrackRecorderSettings = OwningTakeRecorderSource->GetTrackRecorderSettings();
+
 		if (bRecordInWorldSpace && AttachParent && OwningTakeRecorderSource)
 		{
 			// We capture world space transforms for actors if they're attached, but we're not recording the attachment parent
@@ -289,6 +291,8 @@ void UMovieSceneAnimationTrackRecorder::RecordSampleImpl(const FQualifiedFrameTi
 		RecordingSettings.bRecordInWorldSpace = bRecordInWorldSpace;
 		RecordingSettings.bRemoveRootAnimation = bRootWasRemoved;
 		RecordingSettings.bCheckDeltaTimeAtBeginning = false;
+		RecordingSettings.IncludeAnimationNames = TrackRecorderSettings.IncludeAnimationNames;
+		RecordingSettings.ExcludeAnimationNames = TrackRecorderSettings.ExcludeAnimationNames;
 		AnimationRecorder.Init(SkeletalMeshComponent.Get(), AnimSequence.Get(), &AnimationSerializer, RecordingSettings);
 		AnimationRecorder.BeginRecording();
 	}
diff --git a/Engine/Plugins/VirtualProduction/Takes/Source/TakeTrackRecorders/Public/TrackRecorders/IMovieSceneTrackRecorderHost.h b/Engine/Plugins/VirtualProduction/Takes/Source/TakeTrackRecorders/Public/TrackRecorders/IMovieSceneTrackRecorderHost.h
--- a/Engine/Plugins/VirtualProduction/Takes/Source/TakeTrackRecorders/Public/TrackRecorders/IMovieSceneTrackRecorderHost.h
+++ b/Engine/Plugins/VirtualProduction/Takes/Source/TakeTrackRecorders/Public/TrackRecorders/IMovieSceneTrackRecorderHost.h
@@ -52,6 +52,8 @@ struct FTrackRecorderSettings
 	float ReduceKeysTolerance;
 
 	TArray<FTakeRecorderTrackSettings> DefaultTracks;
+	TArray<FString> IncludeAnimationNames;
+	TArray<FString> ExcludeAnimationNames;
 
 	static bool IsDefaultPropertyTrack(UObject* InObjectToRecord, const FString& InPropertyPath, const TArray<FTakeRecorderTrackSettings>& DefaultTracks)
 	{
diff --git a/Engine/Source/Editor/SequenceRecorder/Private/AnimationRecorder.cpp b/Engine/Source/Editor/SequenceRecorder/Private/AnimationRecorder.cpp
--- a/Engine/Source/Editor/SequenceRecorder/Private/AnimationRecorder.cpp
+++ b/Engine/Source/Editor/SequenceRecorder/Private/AnimationRecorder.cpp
@@ -320,6 +320,35 @@ void FAnimationRecorder::ProcessNotifies()
 	}
 }
 
+bool FAnimationRecorder::ShouldSkipName(const FName& InName) const
+{
+	bool bShouldSkipName = false;
+			
+	for (const FString& ExcludeAnimationName : ExcludeAnimationNames)
+	{
+		if (InName.ToString().Contains(ExcludeAnimationName))
+		{
+			bShouldSkipName = true;
+			break;
+		}
+	}
+
+	if (IncludeAnimationNames.Num() != 0)
+	{
+		bShouldSkipName = true;
+		for (const FString& IncludeAnimationName : IncludeAnimationNames)
+		{
+			if (InName.ToString().Contains(IncludeAnimationName))
+			{
+				bShouldSkipName = false;
+				break;
+			}
+		}
+	}
+
+	return bShouldSkipName;
+}
+
 UAnimSequence* FAnimationRecorder::StopRecord(bool bShowMessage)
 {
 	double StartTime, ElapsedTime = 0;
@@ -356,11 +385,23 @@ UAnimSequence* FAnimationRecorder::StopRecord(bool bShowMessage)
 						const bool bMaterialCurve = CurveMetaData->Type.bMaterial;
 						const bool bAttributeCurve = !bMorphTarget && !bMaterialCurve;
 						
+						const FSmartNameMapping* Mapping = SkeletonObj->GetSmartNameContainer(USkeleton::AnimCurveMappingName);
+						FSmartName SmartName;
+
+						bool bShouldSkipName = false;
+						if (Mapping && Mapping->FindSmartNameByUID(CurveUID, SmartName))
+						{
+							bShouldSkipName = ShouldSkipName(SmartName.DisplayName);
+						}
+
 						const bool bSkipCurve = (bMorphTarget && !bRecordMorphTargets) ||
 												(bAttributeCurve && !bRecordAttributeCurves) ||
-												(bMaterialCurve && !bRecordMaterialCurves);
+												(bMaterialCurve && !bRecordMaterialCurves) ||
+												bShouldSkipName;
+
 						if (bSkipCurve)
 						{
+							UE_LOG(LogAnimation, Log, TEXT("Animation Recorder skipping curve: %s"), *SmartName.DisplayName.ToString());
 							continue;
 						}
 					}
@@ -434,7 +475,27 @@ UAnimSequence* FAnimationRecorder::StopRecord(bool bShowMessage)
 			const FBoneAnimationTrack& AnimationTrack = BoneAnimationTracks[TrackIndex];
 			const FRawAnimSequenceTrack& RawTrack = RawTracks[TrackIndex];
 
-			Controller.SetBoneTrackKeys(AnimationTrack.Name, RawTrack.PosKeys, RawTrack.RotKeys, RawTrack.ScaleKeys);
+			FName BoneName = AnimationTrack.Name;
+
+			bool bShouldSkipName = ShouldSkipName(AnimationTrack.Name);
+
+			if (!bShouldSkipName)
+			{
+				Controller.SetBoneTrackKeys(BoneName, RawTrack.PosKeys, RawTrack.RotKeys, RawTrack.ScaleKeys);
+			}
+			else
+			{
+				TArray<FVector3f> SinglePosKey;
+				SinglePosKey.Add(RawTrack.PosKeys[0]);
+				TArray<FQuat4f> SingleRotKey;
+				SingleRotKey.Add(RawTrack.RotKeys[0]);
+				TArray<FVector3f> SingleScaleKey;
+				SingleScaleKey.Add(RawTrack.ScaleKeys[0]);
+
+				Controller.SetBoneTrackKeys(BoneName, SinglePosKey, SingleRotKey, SingleScaleKey);
+
+				UE_LOG(LogAnimation, Log, TEXT("Animation Recorder skipping bone: %s"), *BoneName.ToString());
+			}
 		}
 
 		Controller.NotifyPopulated();
@@ -1055,6 +1116,8 @@ void FAnimRecorderInstance::InitInternal(USkeletalMeshComponent* InComponent, co
 	Recorder->bRecordMorphTargets = Settings.bRecordMorphTargets;
 	Recorder->bRecordAttributeCurves = Settings.bRecordAttributeCurves;
 	Recorder->bRecordMaterialCurves = Settings.bRecordMaterialCurves;
+	Recorder->IncludeAnimationNames = Settings.IncludeAnimationNames;
+	Recorder->ExcludeAnimationNames = Settings.ExcludeAnimationNames;
 
 	if (InComponent)
 	{
diff --git a/Engine/Source/Editor/SequenceRecorder/Public/AnimationRecorder.h b/Engine/Source/Editor/SequenceRecorder/Public/AnimationRecorder.h
--- a/Engine/Source/Editor/SequenceRecorder/Public/AnimationRecorder.h
+++ b/Engine/Source/Editor/SequenceRecorder/Public/AnimationRecorder.h
@@ -141,6 +141,10 @@ public:
 	uint8 bRecordAttributeCurves : 1;
 	/** Whether or not to record material curves*/
 	uint8 bRecordMaterialCurves : 1;
+	/** Include list */
+	TArray<FString> IncludeAnimationNames;
+	/** Exclude list */
+	TArray<FString> ExcludeAnimationNames;
 public:
 	/** Helper function to get space bases depending on master pose component */
 	static void GetBoneTransforms(USkeletalMeshComponent* Component, TArray<FTransform>& BoneTransforms);
@@ -152,6 +156,8 @@ private:
 
 	void ProcessNotifies();
 
+	bool ShouldSkipName(const FName& InName) const;
+
 	// recording curve data 
 	struct FBlendedCurve
 	{
diff --git a/Engine/Source/Runtime/Engine/Classes/Animation/AnimationRecordingSettings.h b/Engine/Source/Runtime/Engine/Classes/Animation/AnimationRecordingSettings.h
--- a/Engine/Source/Runtime/Engine/Classes/Animation/AnimationRecordingSettings.h
+++ b/Engine/Source/Runtime/Engine/Classes/Animation/AnimationRecordingSettings.h
@@ -84,4 +84,12 @@ struct ENGINE_API FAnimationRecordingSettings
 	/** Whether or not to record material curves */
 	UPROPERTY(EditAnywhere, Category = "Settings")
 	bool bRecordMaterialCurves;
+
+	/** Include only the animation bones/curves that match this list */
+	UPROPERTY(EditAnywhere, Category = "Settings")
+	TArray<FString> IncludeAnimationNames;
+
+	/** Exclude all animation bones/curves that match this list */
+	UPROPERTY(EditAnywhere, Category = "Settings")
+	TArray<FString> ExcludeAnimationNames;
 };
