Animation Recorder: Remove bone tracks if recording with transforms turned off.

diff --git a/Engine/Source/Editor/SequenceRecorder/Private/AnimationRecorder.cpp b/Engine/Source/Editor/SequenceRecorder/Private/AnimationRecorder.cpp
--- a/Engine/Source/Editor/SequenceRecorder/Private/AnimationRecorder.cpp
+++ b/Engine/Source/Editor/SequenceRecorder/Private/AnimationRecorder.cpp
@@ -498,9 +498,15 @@ UAnimSequence* FAnimationRecorder::StopRecord(bool bShowMessage)
 			}
 		}
 
+		if (bRecordTransforms == false)
+		{
+			Controller.RemoveAllBoneTracks();
+		}
+
 		Controller.NotifyPopulated();
 		Controller.CloseBracket();
 
+
 		AnimationObject->MarkPackageDirty();
 		
 		// save the package to disk, for convenience and so we can run this in standalone mode
@@ -879,19 +885,19 @@ bool FAnimationRecorder::Record(USkeletalMeshComponent* Component, FTransform co
 
 				if (bRecordTransforms)
 				{
-				FTransform LocalTransform = SpacesBases[BoneIndex];
-				if ( ParentIndex != INDEX_NONE )
-				{
-					LocalTransform.SetToRelativeTransform(SpacesBases[ParentIndex]);
-				}
-				// if record local to world, we'd like to consider component to world to be in root
-				else
-				{
-					if (bRecordLocalToWorld)
+					FTransform LocalTransform = SpacesBases[BoneIndex];
+					if ( ParentIndex != INDEX_NONE )
 					{
-						LocalTransform *= ComponentToWorld;
+						LocalTransform.SetToRelativeTransform(SpacesBases[ParentIndex]);
+					}
+					// if record local to world, we'd like to consider component to world to be in root
+					else
+					{
+						if (bRecordLocalToWorld)
+						{
+							LocalTransform *= ComponentToWorld;
+						}
 					}
-				}
 
 					RawTrack.PosKeys.Add(FVector3f(LocalTransform.GetTranslation()));
 					RawTrack.RotKeys.Add(FQuat4f(LocalTransform.GetRotation()));
