Animation Recorder: Don't reset to default compression if compression is already set.
#trivial

diff --git a/Engine/Source/Editor/SequenceRecorder/Private/AnimationRecorder.cpp b/Engine/Source/Editor/SequenceRecorder/Private/AnimationRecorder.cpp
--- a/Engine/Source/Editor/SequenceRecorder/Private/AnimationRecorder.cpp
+++ b/Engine/Source/Editor/SequenceRecorder/Private/AnimationRecorder.cpp
@@ -235,7 +235,10 @@ void FAnimationRecorder::StartRecord(USkeletalMeshComponent* Component, UAnimSeq
 	TimePassed = 0.0;
 	AnimationObject = InAnimationObject;
 
-	AnimationObject->BoneCompressionSettings = FAnimationUtils::GetDefaultAnimationRecorderBoneCompressionSettings();
+	if (AnimationObject->BoneCompressionSettings.IsNull())
+	{
+		AnimationObject->BoneCompressionSettings = FAnimationUtils::GetDefaultAnimationRecorderBoneCompressionSettings();
+	}
 
 	FAnimationRecorder::GetBoneTransforms(Component, PreviousSpacesBases);
 	PreviousAnimCurves = Component->GetAnimationCurves();
