Fixed crash cancelling an animation export

Popup for the animation recorder allowed access to temporary/invalid objects that were due to be garbage collected
Switched MPK/GC of discarded assets to just rename to the transient package on cancellation
Added optional parameter to not display the 'animation recorded' popup as this is an internal asset creation

diff --git a/Engine/Source/Editor/AnimationEditor/Private/AnimationEditor.cpp b/Engine/Source/Editor/AnimationEditor/Private/AnimationEditor.cpp
--- a/Engine/Source/Editor/AnimationEditor/Private/AnimationEditor.cpp
+++ b/Engine/Source/Editor/AnimationEditor/Private/AnimationEditor.cpp
@@ -835,7 +835,7 @@ void FAnimationEditor::HandleAnimationSequenceBrowserCreated(const TSharedRef<IA
 bool FAnimationEditor::RecordMeshToAnimation(USkeletalMeshComponent* PreviewComponent, UAnimSequence* NewAsset) const
 {
 	ISequenceRecorder& RecorderModule = FModuleManager::Get().LoadModuleChecked<ISequenceRecorder>("SequenceRecorder");
-	return RecorderModule.RecordSingleNodeInstanceToAnimation(PreviewComponent, NewAsset);
+	return RecorderModule.RecordSingleNodeInstanceToAnimation(PreviewComponent, NewAsset, /*bShowMessage*/false);
 }
 
 #undef LOCTEXT_NAMESPACE
diff --git a/Engine/Source/Editor/SequenceRecorder/Private/SequenceRecorderModule.cpp b/Engine/Source/Editor/SequenceRecorder/Private/SequenceRecorderModule.cpp
--- a/Engine/Source/Editor/SequenceRecorder/Private/SequenceRecorderModule.cpp
+++ b/Engine/Source/Editor/SequenceRecorder/Private/SequenceRecorderModule.cpp
@@ -795,9 +795,9 @@ class FSequenceRecorderModule : public ISequenceRecorder, private FSelfRegisteri
 		}
 	}
 
-	virtual bool RecordSingleNodeInstanceToAnimation(USkeletalMeshComponent* PreviewComponent, UAnimSequence* NewAsset) override
+	virtual bool RecordSingleNodeInstanceToAnimation(USkeletalMeshComponent* PreviewComponent, UAnimSequence* NewAsset, bool bShowMessage) override
 	{
-		return SequenceRecorderUtils::RecordSingleNodeInstanceToAnimation(PreviewComponent, NewAsset);
+		return SequenceRecorderUtils::RecordSingleNodeInstanceToAnimation(PreviewComponent, NewAsset, bShowMessage);
 	}
 
 #if WITH_EDITOR
diff --git a/Engine/Source/Editor/SequenceRecorder/Private/SequenceRecorderUtils.cpp b/Engine/Source/Editor/SequenceRecorder/Private/SequenceRecorderUtils.cpp
--- a/Engine/Source/Editor/SequenceRecorder/Private/SequenceRecorderUtils.cpp
+++ b/Engine/Source/Editor/SequenceRecorder/Private/SequenceRecorderUtils.cpp
@@ -36,7 +36,7 @@ AActor* GetAttachment(AActor* InActor, FName& SocketName, FName& ComponentName)
 	return AttachedActor;
 }
 
-bool RecordSingleNodeInstanceToAnimation(USkeletalMeshComponent* PreviewComponent, UAnimSequence* NewAsset)
+bool RecordSingleNodeInstanceToAnimation(USkeletalMeshComponent* PreviewComponent, UAnimSequence* NewAsset, bool bShowMessage)
 {
 	UAnimSingleNodeInstance* SingleNodeInstance = (PreviewComponent) ? Cast<UAnimSingleNodeInstance>(PreviewComponent->GetAnimInstance()) : nullptr;
 	if (SingleNodeInstance && NewAsset)
@@ -77,7 +77,7 @@ bool RecordSingleNodeInstanceToAnimation(USkeletalMeshComponent* PreviewComponen
 			RecordMesh(PreviewComponent, SingleNodeInstance, AnimRecorder, Length, Remainder);
 		}
 
-		AnimRecorder.FinishRecording(true);
+		AnimRecorder.FinishRecording(bShowMessage);
 		return true;
 	}
 
diff --git a/Engine/Source/Editor/SequenceRecorder/Public/ISequenceRecorder.h b/Engine/Source/Editor/SequenceRecorder/Public/ISequenceRecorder.h
--- a/Engine/Source/Editor/SequenceRecorder/Public/ISequenceRecorder.h
+++ b/Engine/Source/Editor/SequenceRecorder/Public/ISequenceRecorder.h
@@ -172,6 +172,7 @@ public:
 	*
 	* @param: PreviewComponent - this component should contains SingleNodeInstance with time-line based asset, currently support AnimSequence or AnimComposite
 	* @param: NewAsset - this is the asset that should be recorded. This will reset all animation data internally
+	* @param: bShowMessage - show a notification of new asset creation
 	*/
-	virtual bool RecordSingleNodeInstanceToAnimation(USkeletalMeshComponent* PreviewComponent, UAnimSequence* NewAsset) = 0;
+	virtual bool RecordSingleNodeInstanceToAnimation(USkeletalMeshComponent* PreviewComponent, UAnimSequence* NewAsset, bool bShowMessage = true) = 0;
 };
diff --git a/Engine/Source/Editor/SequenceRecorder/Public/SequenceRecorderUtils.h b/Engine/Source/Editor/SequenceRecorder/Public/SequenceRecorderUtils.h
--- a/Engine/Source/Editor/SequenceRecorder/Public/SequenceRecorderUtils.h
+++ b/Engine/Source/Editor/SequenceRecorder/Public/SequenceRecorderUtils.h
@@ -203,8 +203,9 @@ SEQUENCERECORDER_API AActor* GetAttachment(AActor* InActor, FName& SocketName, F
  * 
  * @param: PreviewComponent - this component should contains SingleNodeInstance with time-line based asset, currently support AnimSequence or AnimComposite
  * @param: NewAsset - this is the asset that should be recorded. This will reset all animation data internally
+ * @param: bShowMessage - show a notification of new asset creation
  */
-SEQUENCERECORDER_API bool RecordSingleNodeInstanceToAnimation(USkeletalMeshComponent* PreviewComponent, UAnimSequence* NewAsset);
+SEQUENCERECORDER_API bool RecordSingleNodeInstanceToAnimation(USkeletalMeshComponent* PreviewComponent, UAnimSequence* NewAsset, bool bShowMessage = true);
 
 /**
  * Get timecode source
diff --git a/Engine/Source/Editor/UnrealEd/Private/AnimationEditorUtils.cpp b/Engine/Source/Editor/UnrealEd/Private/AnimationEditorUtils.cpp
--- a/Engine/Source/Editor/UnrealEd/Private/AnimationEditorUtils.cpp
+++ b/Engine/Source/Editor/UnrealEd/Private/AnimationEditorUtils.cpp
@@ -518,16 +518,13 @@ namespace AnimationEditorUtils
 		{
 			if (!AssetCreated.Execute(ObjectsToSync))
 			{
-				//Destroy the assets we just create
+				// Rename the objects we created out of the way
 				for (UObject* ObjectToDelete : ObjectsToSync)
 				{
 					// Notify the asset registry
 					FAssetRegistryModule::AssetDeleted(ObjectToDelete);
-					ObjectToDelete->ClearFlags(RF_Standalone | RF_Public);
-					ObjectToDelete->RemoveFromRoot();
-					ObjectToDelete->MarkAsGarbage();
+					ObjectToDelete->Rename(nullptr, GetTransientPackage(), REN_DontCreateRedirectors | REN_ForceNoResetLoaders | REN_NonTransactional);
 				}
-				CollectGarbage(GARBAGE_COLLECTION_KEEPFLAGS);
 			}
 		}
 	}
