Sequence Recorder: Hide Sequence Recorder from menu. Can be re-enabled with cvar "SequenceRecorder"

For 5.1, this just hides the Sequence Recorder from the menu so you can't bring it up. There is already a big red bar in the UI that says it will be removed at some point. For 5.2, we will remove the Sequence Recorder UI entirely. This doesn't change anything about recording from Persona.

diff --git a/Engine/Source/Editor/SequenceRecorder/Private/SequenceRecorderModule.cpp b/Engine/Source/Editor/SequenceRecorder/Private/SequenceRecorderModule.cpp
--- a/Engine/Source/Editor/SequenceRecorder/Private/SequenceRecorderModule.cpp
+++ b/Engine/Source/Editor/SequenceRecorder/Private/SequenceRecorderModule.cpp
@@ -173,13 +173,16 @@ class FSequenceRecorderModule : public ISequenceRecorder, private FSelfRegisteri
 			LevelEditorModule.OnCaptureSingleFrameAnimSequence().BindStatic(&FSequenceRecorderModule::HandleCaptureSingleFrameAnimSequence);
 
 			// register standalone UI
-			auto RegisterTabSpawner = []()
+			auto RegisterTabSpawner = [this]()
 			{
-				FGlobalTabmanager::Get()->RegisterNomadTabSpawner(SequenceRecorderTabName, FOnSpawnTab::CreateStatic(&FSequenceRecorderModule::SpawnSequenceRecorderTab))
+				FTabSpawnerEntry& Entry = FGlobalTabmanager::Get()->RegisterNomadTabSpawner(SequenceRecorderTabName, FOnSpawnTab::CreateStatic(&FSequenceRecorderModule::SpawnSequenceRecorderTab))
 				.SetGroup(WorkspaceMenu::GetMenuStructure().GetLevelEditorCinematicsCategory())
 				.SetDisplayName(LOCTEXT("SequenceRecorderTabTitle", "Sequence Recorder"))
 				.SetTooltipText(LOCTEXT("SequenceRecorderTooltipText", "Open the Sequence Recorder tab."))
 				.SetIcon(FSlateIcon(FAppStyle::GetAppStyleSetName(), "SequenceRecorder.TabIcon"));
+
+				TabSpawnerEntry = &Entry;
+				TabSpawnerEntry->SetMenuType(ETabSpawnerMenuType::Hidden);
 			};
 			FLevelEditorModule* LocalLevelEditorModule = FModuleManager::GetModulePtr<FLevelEditorModule>(TEXT("LevelEditor"));
 			if (LocalLevelEditorModule && LocalLevelEditorModule->GetLevelEditorTabManager())
@@ -886,8 +889,20 @@ class FSequenceRecorderModule : public ISequenceRecorder, private FSelfRegisteri
 	FDelegateHandle AudioFactoryHandle;
 
 	TWeakPtr<SDockTab> SequenceRecorderTab;
+public:
+	FTabSpawnerEntry* TabSpawnerEntry;
 };
 
+static FAutoConsoleCommand CVarEnableSequenceRecorder(
+	TEXT("SequenceRecorder"),
+	TEXT("Enables the Sequence Recorder tab"),
+	FConsoleCommandWithArgsDelegate::CreateLambda([](const TArray<FString>& Args)
+		{
+			FSequenceRecorderModule& SequenceRecorder = FModuleManager::GetModuleChecked<FSequenceRecorderModule>("SequenceRecorder");
+			SequenceRecorder.TabSpawnerEntry->SetMenuType(ETabSpawnerMenuType::Enabled);
+		}),
+	ECVF_Default);
+
 IMPLEMENT_MODULE( FSequenceRecorderModule, SequenceRecorder )
 
 #undef LOCTEXT_NAMESPACE
