FNames containing asset paths are deprecated. FSoftObjectPath should be used instead.
Fixups for LevelSequence/SequenceRecorder.

diff --git a/Engine/Plugins/MovieScene/LevelSequenceEditor/Source/LevelSequenceEditor/Private/Misc/LevelSequenceEditorHelpers.cpp b/Engine/Plugins/MovieScene/LevelSequenceEditor/Source/LevelSequenceEditor/Private/Misc/LevelSequenceEditorHelpers.cpp
--- a/Engine/Plugins/MovieScene/LevelSequenceEditor/Source/LevelSequenceEditor/Private/Misc/LevelSequenceEditorHelpers.cpp
+++ b/Engine/Plugins/MovieScene/LevelSequenceEditor/Source/LevelSequenceEditor/Private/Misc/LevelSequenceEditorHelpers.cpp
@@ -185,7 +185,7 @@ private:
 		FullPath /= LevelSequenceSettings->MasterSequenceName;
 		FullPath /= LevelSequenceSettings->MasterSequenceName;
 		FullPath += LevelSequenceSettings->MasterSequenceSuffix;
-		FAssetData AssetData = AssetRegistryModule.Get().GetAssetByObjectPath(FName(*FullPath));
+		FAssetData AssetData = AssetRegistryModule.Get().GetAssetByObjectPath(FSoftObjectPath(FullPath));
 		if (AssetData.IsValid())
 		{
 			return LOCTEXT("MasterSequenceExists", "Warning: Master Sequence Exists");
diff --git a/Engine/Source/Editor/SequenceRecorder/Private/SequenceRecorderUtils.cpp b/Engine/Source/Editor/SequenceRecorder/Private/SequenceRecorderUtils.cpp
--- a/Engine/Source/Editor/SequenceRecorder/Private/SequenceRecorderUtils.cpp
+++ b/Engine/Source/Editor/SequenceRecorder/Private/SequenceRecorderUtils.cpp
@@ -107,7 +107,7 @@ FString MakeNewGroupName(const FString& BaseAssetPath, const FString& BaseAssetN
 		// If the existing base asset name doesn't conflict, use it
 		FString AssetPath = BaseAssetPath / BaseAssetName / BaseAssetName + Dot + BaseAssetName;
 
-		if (!ExistingGroupNames.Contains(FName(*BaseAssetName)) && !AssetRegistryModule.Get().GetAssetByObjectPath(*AssetPath).IsValid())
+		if (!ExistingGroupNames.Contains(FName(*BaseAssetName)) && !AssetRegistryModule.Get().GetAssetByObjectPath(FSoftObjectPath(AssetPath)).IsValid())
 		{
 			return BaseAssetName;
 		}
@@ -137,7 +137,7 @@ FString MakeNewGroupName(const FString& BaseAssetPath, const FString& BaseAssetN
 
 		FString AssetPath = BaseAssetPath / NewAssetName / NewAssetName + Dot + NewAssetName;
 
-		if (!ExistingGroupNames.Contains(FName(*NewAssetName)) && !AssetRegistryModule.Get().GetAssetByObjectPath(*AssetPath).IsValid())
+		if (!ExistingGroupNames.Contains(FName(*NewAssetName)) && !AssetRegistryModule.Get().GetAssetByObjectPath(FSoftObjectPath(AssetPath)).IsValid())
 		{
 			return NewAssetName;
 		}
diff --git a/Engine/Source/Editor/SequenceRecorder/Public/SequenceRecorderUtils.h b/Engine/Source/Editor/SequenceRecorder/Public/SequenceRecorderUtils.h
--- a/Engine/Source/Editor/SequenceRecorder/Public/SequenceRecorderUtils.h
+++ b/Engine/Source/Editor/SequenceRecorder/Public/SequenceRecorderUtils.h
@@ -35,7 +35,7 @@ static AssetType* MakeNewAsset(const FString& BaseAssetPath, const FString& Base
 	AssetPath += Dot + AssetName;
 
 	FAssetRegistryModule& AssetRegistryModule = FModuleManager::LoadModuleChecked<FAssetRegistryModule>(TEXT("AssetRegistry"));
-	FAssetData AssetData = AssetRegistryModule.Get().GetAssetByObjectPath(*AssetPath);
+	FAssetData AssetData = AssetRegistryModule.Get().GetAssetByObjectPath(FSoftObjectPath(AssetPath));
 
 	// if object with same name exists, try a different name until we don't find one
 	int32 ExtensionIndex = 0;
@@ -43,7 +43,7 @@ static AssetType* MakeNewAsset(const FString& BaseAssetPath, const FString& Base
 	{
 		AssetName = FString::Printf(TEXT("%s_%d"), *BaseAssetName, ExtensionIndex);
 		AssetPath = (BaseAssetPath / AssetName) + Dot + AssetName;
-		AssetData = AssetRegistryModule.Get().GetAssetByObjectPath(*AssetPath);
+		AssetData = AssetRegistryModule.Get().GetAssetByObjectPath(FSoftObjectPath(AssetPath));
 
 		ExtensionIndex++;
 	}
@@ -79,7 +79,7 @@ static UObject* DuplicateAsset(const FString& BaseAssetPath, const FString& Base
 	AssetPath += Dot + AssetName;
 
 	FAssetRegistryModule& AssetRegistryModule = FModuleManager::LoadModuleChecked<FAssetRegistryModule>(TEXT("AssetRegistry"));
-	FAssetData AssetData = AssetRegistryModule.Get().GetAssetByObjectPath(*AssetPath);
+	FAssetData AssetData = AssetRegistryModule.Get().GetAssetByObjectPath(FSoftObjectPath(AssetPath));
 
 	// if object with same name exists, try a different name until we don't find one
 	int32 ExtensionIndex = 0;
@@ -87,7 +87,7 @@ static UObject* DuplicateAsset(const FString& BaseAssetPath, const FString& Base
 	{
 		AssetName = FString::Printf(TEXT("%s_%d"), *BaseAssetName, ExtensionIndex);
 		AssetPath = (BaseAssetPath / AssetName) + Dot + AssetName;
-		AssetData = AssetRegistryModule.Get().GetAssetByObjectPath(*AssetPath);
+		AssetData = AssetRegistryModule.Get().GetAssetByObjectPath(FSoftObjectPath(AssetPath));
 
 		ExtensionIndex++;
 	}
@@ -115,7 +115,7 @@ static FString MakeNewAssetName(const FString& BaseAssetPath, const FString& Bas
 	AssetPath += Dot + AssetName;
 
 	FAssetRegistryModule& AssetRegistryModule = FModuleManager::LoadModuleChecked<FAssetRegistryModule>(TEXT("AssetRegistry"));
-	FAssetData AssetData = AssetRegistryModule.Get().GetAssetByObjectPath(*AssetPath);
+	FAssetData AssetData = AssetRegistryModule.Get().GetAssetByObjectPath(FSoftObjectPath(AssetPath));
 
 	// if object with same name exists, try a different name until we don't find one
 	int32 ExtensionIndex = 0;
@@ -123,7 +123,7 @@ static FString MakeNewAssetName(const FString& BaseAssetPath, const FString& Bas
 	{
 		AssetName = FString::Printf(TEXT("%s_%d"), *BaseAssetName, ExtensionIndex);
 		AssetPath = (BaseAssetPath / AssetName) + Dot + AssetName;
-		AssetData = AssetRegistryModule.Get().GetAssetByObjectPath(*AssetPath);
+		AssetData = AssetRegistryModule.Get().GetAssetByObjectPath(FSoftObjectPath(AssetPath));
 
 		ExtensionIndex++;
 	}
@@ -150,7 +150,7 @@ static bool DoesTakeExist(const FString& AssetPath, const FString& ActorName, co
 
 	const FString Dot(TEXT("."));
 	FString FullPath = TakePath + Dot + TakeName;
-	FAssetData AssetData = AssetRegistryModule.Get().GetAssetByObjectPath(*FullPath);
+	FAssetData AssetData = AssetRegistryModule.Get().GetAssetByObjectPath(FSoftObjectPath(FullPath));
 
 	return AssetData.IsValid();
 }
diff --git a/Engine/Source/Runtime/LevelSequence/Private/LevelSequence.cpp b/Engine/Source/Runtime/LevelSequence/Private/LevelSequence.cpp
--- a/Engine/Source/Runtime/LevelSequence/Private/LevelSequence.cpp
+++ b/Engine/Source/Runtime/LevelSequence/Private/LevelSequence.cpp
@@ -404,10 +404,10 @@ bool ULevelSequence::CanPossessObject(UObject& Object, UObject* InPlaybackContex
 
 void ULevelSequence::LocateBoundObjects(const FGuid& ObjectId, UObject* Context, TArray<UObject*, TInlineAllocator<1>>& OutObjects) const
 {
-	LocateBoundObjects(ObjectId, Context, NAME_None, OutObjects);
+	LocateBoundObjects(ObjectId, Context, {}, OutObjects);
 }
 
-void ULevelSequence::LocateBoundObjects(const FGuid& ObjectId, UObject* Context, FName StreamedLevelAssetPath, TArray<UObject*, TInlineAllocator<1>>& OutObjects) const
+void ULevelSequence::LocateBoundObjects(const FGuid& ObjectId, UObject* Context, const FTopLevelAssetPath& StreamedLevelAssetPath, TArray<UObject*, TInlineAllocator<1>>& OutObjects) const
 {
 	// Handle legacy object references
 	UObject* Object = Context ? ObjectReferences.ResolveBinding(ObjectId, Context) : nullptr;
diff --git a/Engine/Source/Runtime/LevelSequence/Private/LevelSequenceBindingReference.cpp b/Engine/Source/Runtime/LevelSequence/Private/LevelSequenceBindingReference.cpp
--- a/Engine/Source/Runtime/LevelSequence/Private/LevelSequenceBindingReference.cpp
+++ b/Engine/Source/Runtime/LevelSequence/Private/LevelSequenceBindingReference.cpp
@@ -43,7 +43,7 @@ FLevelSequenceBindingReference::FLevelSequenceBindingReference(UObject* InObject
 	}
 }
 
-UObject* FLevelSequenceBindingReference::Resolve(UObject* InContext, FName StreamedLevelAssetPath) const
+UObject* FLevelSequenceBindingReference::Resolve(UObject* InContext, const FTopLevelAssetPath& StreamedLevelAssetPath) const
 {
 	if (InContext && InContext->IsA<AActor>())
 	{
@@ -57,7 +57,7 @@ UObject* FLevelSequenceBindingReference::Resolve(UObject* InContext, FName Strea
 			return FindObject<UObject>(InContext, *ObjectPath, false);
 		}
 	}
-	else if (InContext && InContext->IsA<ULevel>() && StreamedLevelAssetPath != NAME_None && ExternalObjectPath.GetAssetPathName() == StreamedLevelAssetPath)
+	else if (InContext && InContext->IsA<ULevel>() && StreamedLevelAssetPath.IsValid() && ExternalObjectPath.GetAssetPath() == StreamedLevelAssetPath)
 	{
 		if (UE::IsSavingPackage(nullptr) || IsGarbageCollecting())
 		{
@@ -247,7 +247,7 @@ void FLevelSequenceBindingReferences::RemoveObjects(const FGuid& ObjectId, const
 
 	for (int32 ReferenceIndex = 0; ReferenceIndex < ReferenceArray->References.Num(); )
 	{
-		UObject* ResolvedObject = ReferenceArray->References[ReferenceIndex].Resolve(InContext, NAME_None);
+		UObject* ResolvedObject = ReferenceArray->References[ReferenceIndex].Resolve(InContext, FTopLevelAssetPath());
 
 		if (InObjects.Contains(ResolvedObject))
 		{
@@ -270,7 +270,7 @@ void FLevelSequenceBindingReferences::RemoveInvalidObjects(const FGuid& ObjectId
 
 	for (int32 ReferenceIndex = 0; ReferenceIndex < ReferenceArray->References.Num(); )
 	{
-		UObject* ResolvedObject = ReferenceArray->References[ReferenceIndex].Resolve(InContext, NAME_None);
+		UObject* ResolvedObject = ReferenceArray->References[ReferenceIndex].Resolve(InContext, FTopLevelAssetPath());
 
 		if (!IsValid(ResolvedObject))
 		{
@@ -283,7 +283,7 @@ void FLevelSequenceBindingReferences::RemoveInvalidObjects(const FGuid& ObjectId
 	}
 }
 
-void FLevelSequenceBindingReferences::ResolveBinding(const FGuid& ObjectId, UObject* InContext, FName StreamedLevelAssetPath, TArray<UObject*, TInlineAllocator<1>>& OutObjects) const
+void FLevelSequenceBindingReferences::ResolveBinding(const FGuid& ObjectId, UObject* InContext, const FTopLevelAssetPath& StreamedLevelAssetPath, TArray<UObject*, TInlineAllocator<1>>& OutObjects) const
 {
 	if (const FLevelSequenceBindingReferenceArray* ReferenceArray = BindingIdToReferences.Find(ObjectId))
 	{
diff --git a/Engine/Source/Runtime/LevelSequence/Private/LevelSequencePlayer.cpp b/Engine/Source/Runtime/LevelSequence/Private/LevelSequencePlayer.cpp
--- a/Engine/Source/Runtime/LevelSequence/Private/LevelSequencePlayer.cpp
+++ b/Engine/Source/Runtime/LevelSequence/Private/LevelSequencePlayer.cpp
@@ -80,7 +80,7 @@ ULevelSequencePlayer* ULevelSequencePlayer::CreateLevelSequencePlayer(UObject* W
 void ULevelSequencePlayer::Initialize(ULevelSequence* InLevelSequence, ULevel* InLevel, const FMovieSceneSequencePlaybackSettings& Settings, const FLevelSequenceCameraSettings& InCameraSettings)
 {
 	// Never use the level to resolve bindings unless we're playing back within a streamed or instanced level
-	StreamedLevelAssetPath = NAME_None;
+	StreamedLevelAssetPath = FTopLevelAssetPath();
 
 	World = InLevel->OwningWorld;
 	Level = InLevel;
@@ -96,9 +96,7 @@ void ULevelSequencePlayer::Initialize(ULevelSequence* InLevelSequence, ULevel* I
 		int32 SlashPos = 0;
 		if (StreamedLevelPackage.FindLastChar('/', SlashPos) && SlashPos < StreamedLevelPackage.Len()-1)
 		{
-			// Construct the asset path by appending .MapName to the end for efficient comparison with FSoftObjectPath::GetAssetPathName
-			const TCHAR* Pair[] = { *StreamedLevelPackage, &StreamedLevelPackage[SlashPos+1] };
-			StreamedLevelAssetPath = *FString::Join(Pair, TEXT("."));
+			StreamedLevelAssetPath = FTopLevelAssetPath(*StreamedLevelPackage, &StreamedLevelPackage[SlashPos+1]);
 		}
 	}
 
@@ -112,7 +110,7 @@ void ULevelSequencePlayer::ResolveBoundObjects(const FGuid& InBindingId, FMovieS
 
 	if (bAllowDefault)
 	{
-		if (StreamedLevelAssetPath != NAME_None && ResolutionContext && ResolutionContext->IsA<UWorld>())
+		if (StreamedLevelAssetPath.IsValid() && ResolutionContext && ResolutionContext->IsA<UWorld>())
 		{
 			ResolutionContext = Level.Get();
 		}
diff --git a/Engine/Source/Runtime/LevelSequence/Public/LevelSequence.h b/Engine/Source/Runtime/LevelSequence/Public/LevelSequence.h
--- a/Engine/Source/Runtime/LevelSequence/Public/LevelSequence.h
+++ b/Engine/Source/Runtime/LevelSequence/Public/LevelSequence.h
@@ -83,7 +83,7 @@ public:
 
 	virtual void PostDuplicate(bool bDuplicateForPIE) override;
 
-	void LocateBoundObjects(const FGuid& ObjectId, UObject* Context, FName StreamedLevelAssetPath, TArray<UObject*, TInlineAllocator<1>>& OutObjects) const;
+	void LocateBoundObjects(const FGuid& ObjectId, UObject* Context, const FTopLevelAssetPath& StreamedLevelAssetPath, TArray<UObject*, TInlineAllocator<1>>& OutObjects) const;
 #if WITH_EDITOR
 
 
diff --git a/Engine/Source/Runtime/LevelSequence/Public/LevelSequenceBindingReference.h b/Engine/Source/Runtime/LevelSequence/Public/LevelSequenceBindingReference.h
--- a/Engine/Source/Runtime/LevelSequence/Public/LevelSequenceBindingReference.h
+++ b/Engine/Source/Runtime/LevelSequence/Public/LevelSequenceBindingReference.h
@@ -36,10 +36,18 @@ struct FLevelSequenceBindingReference
 	 * Resolve this reference within the specified context
 	 *
 	 * @param InContext		The context to resolve the binding within. Either a UWorld, ULevel (when playing in an instanced level) or an AActor where this binding relates to an actor component
-	 * @oaram StreamedLevelAssetPath    The path to the streamed level asset that contains the level sequence actor playing back the sequence. 'None' for any non-instance-level setups.
+	 * @param StreamedLevelAssetPath    The path to the streamed level asset that contains the level sequence actor playing back the sequence. A null path for any non-instance-level setups.
 	 * @return The object (usually an Actor or an ActorComponent).
 	 */
-	LEVELSEQUENCE_API UObject* Resolve(UObject* InContext, FName StreamedLevelAssetPath) const;
+	LEVELSEQUENCE_API UObject* Resolve(UObject* InContext, const FTopLevelAssetPath& StreamedLevelAssetPath) const;
+	
+	UE_DEPRECATED(5.1, "Asset path FNames are deprecated. Use FTopLevelAssetPath instead.")
+	UObject* Resolve(UObject* InContext, FName StreamedLevelAssetPath) const
+	{
+PRAGMA_DISABLE_DEPRECATION_WARNINGS
+		return Resolve(InContext, FTopLevelAssetPath(StreamedLevelAssetPath));
+PRAGMA_ENABLE_DEPRECATION_WARNINGS		
+	}
 
 	/**
 	 * Check whether this binding reference is equal to the specified object
@@ -142,7 +150,7 @@ struct FLevelSequenceBindingReferences
 	 * @oaram StreamedLevelAssetPath    The path to the streamed level asset that contains the level sequence actor playing back the sequence. 'None' for any non-instance-level setups.
 	 * @param OutObjects				Array to populate with resolved object bindings
 	 */
-	void ResolveBinding(const FGuid& ObjectId, UObject* InContext, FName StreamedLevelAssetPath, TArray<UObject*, TInlineAllocator<1>>& OutObjects) const;
+	void ResolveBinding(const FGuid& ObjectId, UObject* InContext, const FTopLevelAssetPath& StreamedLevelAssetPath, TArray<UObject*, TInlineAllocator<1>>& OutObjects) const;
 
 	/**
 	 * Resolve a binding for the specified ID using a given context
diff --git a/Engine/Source/Runtime/LevelSequence/Public/LevelSequencePlayer.h b/Engine/Source/Runtime/LevelSequence/Public/LevelSequencePlayer.h
--- a/Engine/Source/Runtime/LevelSequence/Public/LevelSequencePlayer.h
+++ b/Engine/Source/Runtime/LevelSequence/Public/LevelSequencePlayer.h
@@ -172,8 +172,8 @@ private:
 	/** The world this player will spawn actors in, if needed */
 	TWeakObjectPtr<ULevel> Level;
 
-	/** The full asset path (/Game/Folder/MapName.MapName) of the streaming level this player resides within. Bindings to actors with the same FSoftObjectPath::GetAssetPathName are resolved within the cached level, rather than globally.. */
-	FName StreamedLevelAssetPath;
+	/** The full asset path (/Game/Folder/MapName.MapName) of the streaming level this player resides within. Bindings to actors with the same FSoftObjectPath::GetAssetPath are resolved within the cached level, rather than globally.. */
+	FTopLevelAssetPath StreamedLevelAssetPath;
 
 	/** The camera settings to use when playing the sequence */
 	FLevelSequenceCameraSettings CameraSettings;
