Animation Recorder: Setting Retarget Source Asset also needs to make sure UpdateRetargetSourceAsset is also called, it's protected and so we need to do what details customization does and force a property changed.

diff --git a/Engine/Source/Editor/SequenceRecorder/Private/AnimationRecorder.cpp b/Engine/Source/Editor/SequenceRecorder/Private/AnimationRecorder.cpp
--- a/Engine/Source/Editor/SequenceRecorder/Private/AnimationRecorder.cpp
+++ b/Engine/Source/Editor/SequenceRecorder/Private/AnimationRecorder.cpp
@@ -289,6 +289,12 @@ void FAnimationRecorder::StartRecord(USkeletalMeshComponent* Component, UAnimSeq
 	if (AnimationObject->RetargetSource == NAME_None)
 	{
 		AnimationObject->RetargetSourceAsset = Component->GetSkeletalMeshAsset();
+		//UpdateRetargetSourceAssetData() is protected so need to do a posteditchagned
+#if WITH_EDITOR
+		FProperty* PropertyChanged = UAnimSequence::StaticClass()->FindPropertyByName(GET_MEMBER_NAME_CHECKED(UAnimSequence, RetargetSourceAsset));
+		FPropertyChangedEvent PropertyUpdateStruct(PropertyChanged);
+		AnimationObject->PostEditChangeProperty(PropertyUpdateStruct);
+#endif
 	}
 
 	// record the first frame
