Exporting an Animation Sequence from a Sequence that has a spawnable causes a crash
#fix add option to disable controller API transactions, disabled for Sequencer exports using Spawnables

diff --git a/Engine/Plugins/MovieScene/SequencerScripting/Source/SequencerScriptingEditor/Private/SequencerTools.cpp b/Engine/Plugins/MovieScene/SequencerScripting/Source/SequencerScriptingEditor/Private/SequencerTools.cpp
--- a/Engine/Plugins/MovieScene/SequencerScripting/Source/SequencerScriptingEditor/Private/SequencerTools.cpp
+++ b/Engine/Plugins/MovieScene/SequencerScripting/Source/SequencerScriptingEditor/Private/SequencerTools.cpp
@@ -354,6 +354,9 @@ bool USequencerToolsFunctionLibrary::ExportAnimSequence(UWorld* World, ULevelSeq
 			// Evaluate at the beginning of the subscene time to ensure that spawnables are created before export
 			FFrameTime StartTime = FFrameRate::TransformTime(UE::MovieScene::DiscreteInclusiveLower(MovieScene->GetPlaybackRange()).Value, MovieScene->GetTickResolution(), MovieScene->GetDisplayRate());
 			Player->SetPlaybackPosition(FMovieSceneSequencePlaybackParams(StartTime, EUpdatePositionMethod::Play));
+
+			// Cannot transact changes due to in-flight spawnable transaction
+			ExportOptions->bTransactRecording = false;
 		}
  
 		USkeletalMeshComponent* SkeletalMeshComp =  GetSkelMeshComponent(Player, Binding);
diff --git a/Engine/Source/Editor/SequenceRecorder/Private/AnimationRecorder.cpp b/Engine/Source/Editor/SequenceRecorder/Private/AnimationRecorder.cpp
--- a/Engine/Source/Editor/SequenceRecorder/Private/AnimationRecorder.cpp
+++ b/Engine/Source/Editor/SequenceRecorder/Private/AnimationRecorder.cpp
@@ -252,14 +252,14 @@ void FAnimationRecorder::StartRecord(USkeletalMeshComponent* Component, UAnimSeq
 	IAnimationDataController& Controller = AnimationObject->GetController();
 	Controller.SetModel(AnimationObject->GetDataModelInterface());
 
-	Controller.OpenBracket(LOCTEXT("StartRecord_Bracket", "Starting Animation Recording"));
+	Controller.OpenBracket(LOCTEXT("StartRecord_Bracket", "Starting Animation Recording"), bTransactRecording);
 
 	Controller.InitializeModel();
 
 	const bool bKeepNotifiesAndCurves = CVarKeepNotifyAndCurvesOnAnimationRecord->GetInt() == 0 ? false : true;
 	if (bKeepNotifiesAndCurves)
 	{
-		Controller.RemoveAllBoneTracks();
+		Controller.RemoveAllBoneTracks(bTransactRecording);
 	}
 	else
 	{
@@ -282,7 +282,7 @@ void FAnimationRecorder::StartRecord(USkeletalMeshComponent* Component, UAnimSeq
 		{
 			// add tracks for the bone existing
 			const FName BoneTreeName = AnimSkeleton->GetReferenceSkeleton().GetBoneName(BoneTreeIndex);
-			Controller.AddBoneCurve(BoneTreeName);			
+			Controller.AddBoneCurve(BoneTreeName, bTransactRecording);			
 			RawTracks.AddDefaulted();
 		}
 	}
@@ -378,8 +378,8 @@ UAnimSequence* FAnimationRecorder::StopRecord(bool bShowMessage)
 		AnimationObject->Interpolation = Interpolation;
 
 		// can't use TimePassed. That is just total time that has been passed, not necessarily match with frame count
-		Controller.SetFrameRate(RecordingRate);
-		Controller.SetNumberOfFrames( FMath::Max(LastFrame.Value,1) );
+		Controller.SetFrameRate(RecordingRate, bTransactRecording);
+		Controller.SetNumberOfFrames( FMath::Max(LastFrame.Value,1), bTransactRecording);
 
 		ProcessNotifies();
 
@@ -439,7 +439,7 @@ UAnimSequence* FAnimationRecorder::StopRecord(bool bShowMessage)
 							// add one and save the cache
 							// give default curve flag for recording 
 							const FAnimationCurveIdentifier CurveId(CurveName, ERawCurveTrackTypes::RCT_Float);
-							Controller.AddCurve(CurveId, AACF_DefaultCurve);
+							Controller.AddCurve(CurveId, AACF_DefaultCurve, bTransactRecording);
 							FloatCurveData = AnimationObject->GetDataModel()->FindFloatCurve(CurveId);
 						}
 
@@ -466,7 +466,7 @@ UAnimSequence* FAnimationRecorder::StopRecord(bool bShowMessage)
 					}
 
 					const FAnimationCurveIdentifier CurveId(FloatCurveData->GetName(), ERawCurveTrackTypes::RCT_Float);
-					Controller.SetCurveKeys(CurveId, Keys);
+					Controller.SetCurveKeys(CurveId, Keys, bTransactRecording);
 				}
 			}
 
@@ -486,7 +486,7 @@ UAnimSequence* FAnimationRecorder::StopRecord(bool bShowMessage)
 
 			if (!bShouldSkipName)
 			{
-				Controller.SetBoneTrackKeys(BoneName, RawTrack.PosKeys, RawTrack.RotKeys, RawTrack.ScaleKeys);
+				Controller.SetBoneTrackKeys(BoneName, RawTrack.PosKeys, RawTrack.RotKeys, RawTrack.ScaleKeys, bTransactRecording);
 			}
 			else
 			{
@@ -497,7 +497,7 @@ UAnimSequence* FAnimationRecorder::StopRecord(bool bShowMessage)
 				TArray<FVector3f> SingleScaleKey;
 				SingleScaleKey.Add(RawTrack.ScaleKeys[0]);
 
-				Controller.SetBoneTrackKeys(BoneName, SinglePosKey, SingleRotKey, SingleScaleKey);
+				Controller.SetBoneTrackKeys(BoneName, SinglePosKey, SingleRotKey, SingleScaleKey, bTransactRecording);
 
 				UE_LOG(LogAnimation, Log, TEXT("Animation Recorder skipping bone: %s"), *BoneName.ToString());
 			}
@@ -505,11 +505,11 @@ UAnimSequence* FAnimationRecorder::StopRecord(bool bShowMessage)
 
 		if (bRecordTransforms == false)
 		{
-			Controller.RemoveAllBoneTracks();
+			Controller.RemoveAllBoneTracks(bTransactRecording);
 		}
 
 		Controller.NotifyPopulated();
-		Controller.CloseBracket();
+		Controller.CloseBracket(bTransactRecording);
 
 
 		AnimationObject->MarkPackageDirty();
@@ -1118,6 +1118,7 @@ void FAnimRecorderInstance::InitInternal(USkeletalMeshComponent* InComponent, co
 	Recorder->bRecordMorphTargets = Settings.bRecordMorphTargets;
 	Recorder->bRecordAttributeCurves = Settings.bRecordAttributeCurves;
 	Recorder->bRecordMaterialCurves = Settings.bRecordMaterialCurves;
+	Recorder->bTransactRecording = Settings.bTransactRecording;
 	Recorder->IncludeAnimationNames = Settings.IncludeAnimationNames;
 	Recorder->ExcludeAnimationNames = Settings.ExcludeAnimationNames;
 
diff --git a/Engine/Source/Editor/SequenceRecorder/Public/AnimationRecorder.h b/Engine/Source/Editor/SequenceRecorder/Public/AnimationRecorder.h
--- a/Engine/Source/Editor/SequenceRecorder/Public/AnimationRecorder.h
+++ b/Engine/Source/Editor/SequenceRecorder/Public/AnimationRecorder.h
@@ -147,6 +147,8 @@ public:
 	TArray<FString> IncludeAnimationNames;
 	/** Exclude list */
 	TArray<FString> ExcludeAnimationNames;
+	/** Whether or not to transact any IAnimationDataController changes */
+	bool bTransactRecording;
 public:
 	/** Helper function to get space bases depending on leader pose component */
 	static void GetBoneTransforms(USkeletalMeshComponent* Component, TArray<FTransform>& BoneTransforms);
diff --git a/Engine/Source/Editor/UnrealEd/Classes/Exporters/AnimSeqExportOption.h b/Engine/Source/Editor/UnrealEd/Classes/Exporters/AnimSeqExportOption.h
--- a/Engine/Source/Editor/UnrealEd/Classes/Exporters/AnimSeqExportOption.h
+++ b/Engine/Source/Editor/UnrealEd/Classes/Exporters/AnimSeqExportOption.h
@@ -66,6 +66,9 @@ public:
 	UPROPERTY(EditAnywhere, BlueprintReadWrite, AdvancedDisplay, Category = Export);
 	FFrameNumber DelayBeforeStart;
 
+	/** Whether or not to transact the animation sequence data recording */
+	UPROPERTY(EditAnywhere, BlueprintReadWrite, AdvancedDisplay, Category = Export)
+	bool bTransactRecording = true;
 
 	void ResetToDefault()
 	{
@@ -79,5 +82,6 @@ public:
 		bEvaluateAllSkeletalMeshComponents = true;
 		WarmUpFrames = 0;
 		DelayBeforeStart = 0;
+		bTransactRecording = true;
 	}
 };
diff --git a/Engine/Source/Runtime/Engine/Classes/Animation/AnimationRecordingSettings.h b/Engine/Source/Runtime/Engine/Classes/Animation/AnimationRecordingSettings.h
--- a/Engine/Source/Runtime/Engine/Classes/Animation/AnimationRecordingSettings.h
+++ b/Engine/Source/Runtime/Engine/Classes/Animation/AnimationRecordingSettings.h
@@ -38,6 +38,7 @@ struct ENGINE_API FAnimationRecordingSettings
 		, bRecordMorphTargets(true)
 		, bRecordAttributeCurves(true)
 		, bRecordMaterialCurves(true)
+		, bTransactRecording(true) 
 	{}
 	
 	/** Whether to record animation in world space, defaults to true */
@@ -91,6 +92,10 @@ struct ENGINE_API FAnimationRecordingSettings
 	UPROPERTY(EditAnywhere, Category = "Settings")
 	bool bRecordMaterialCurves;
 
+	/** Whether or not to transact recording changes */
+	UPROPERTY(EditAnywhere, Category = "Settings")
+	bool bTransactRecording;
+
 	/** Include only the animation bones/curves that match this list */
 	UPROPERTY(EditAnywhere, Category = "Settings")
 	TArray<FString> IncludeAnimationNames;
