Take Recorder: Don't finalize the recording on shutdown if the animation asset is destroyed

diff --git a/Engine/Source/Editor/SequenceRecorder/Private/AnimationRecorder.cpp b/Engine/Source/Editor/SequenceRecorder/Private/AnimationRecorder.cpp
--- a/Engine/Source/Editor/SequenceRecorder/Private/AnimationRecorder.cpp
+++ b/Engine/Source/Editor/SequenceRecorder/Private/AnimationRecorder.cpp
@@ -369,7 +369,7 @@ UAnimSequence* FAnimationRecorder::StopRecord(bool bShowMessage)
 {
 	double StartTime, ElapsedTime = 0;
 
-	if (AnimationObject)
+	if (IsValid(AnimationObject) && !AnimationObject->HasAnyFlags(RF_BeginDestroyed | RF_FinishDestroyed))
 	{
 		IAnimationDataController& Controller = AnimationObject->GetController();
 		int32 NumKeys = LastFrame.Value + 1;
