Force a ::PostEditChange when finishing animation recording (this will trigger any mechanisms hooked off ::PostEditChangeOwner)

diff --git a/Engine/Source/Editor/SequenceRecorder/Private/AnimationRecorder.cpp b/Engine/Source/Editor/SequenceRecorder/Private/AnimationRecorder.cpp
--- a/Engine/Source/Editor/SequenceRecorder/Private/AnimationRecorder.cpp
+++ b/Engine/Source/Editor/SequenceRecorder/Private/AnimationRecorder.cpp
@@ -511,7 +511,7 @@ UAnimSequence* FAnimationRecorder::StopRecord(bool bShowMessage)
 		Controller.NotifyPopulated();
 		Controller.CloseBracket(bTransactRecording);
 
-
+		AnimationObject->PostEditChange();
 		AnimationObject->MarkPackageDirty();
 		
 		// save the package to disk, for convenience and so we can run this in standalone mode
