Fixed recording of anim notify states in looping montage sections

Notify states that overlap looping montage sections now get correctly recorded
Notifies that fire or overlap the initial frame are now correctly recorded
Notifies that are active when recording finishes are now correctly recorded

diff --git a/Engine/Source/Editor/SequenceRecorder/Private/AnimationRecorder.cpp b/Engine/Source/Editor/SequenceRecorder/Private/AnimationRecorder.cpp
--- a/Engine/Source/Editor/SequenceRecorder/Private/AnimationRecorder.cpp
+++ b/Engine/Source/Editor/SequenceRecorder/Private/AnimationRecorder.cpp
@@ -302,41 +302,102 @@ void FAnimationRecorder::StartRecord(USkeletalMeshComponent* Component, UAnimSeq
 #endif
 	}
 
+	// record first-frame notifies
+	if (Component->GetAnimInstance())
+	{
+		RecordNotifies(Component, Component->GetAnimInstance()->NotifyQueue.AnimNotifies, 0.0f, 0.0f);
+	}
+
 	// record the first frame
 	Record(Component, PreviousComponentToWorld, PreviousSpacesBases, PreviousAnimCurves,  0);
 }
 
 void FAnimationRecorder::ProcessNotifies()
 {
-	if (AnimationObject)
+	if (AnimationObject == nullptr)
 	{
-		// Copy recorded notify events, animation its notify array should be empty at this point
-		AnimationObject->Notifies.Append(RecordedNotifyEvents);
-		
-		// build notify tracks - first find how many tracks we want
-		for (FAnimNotifyEvent& Event : AnimationObject->Notifies)
+		return;
+	}
+
+	// Get play length to clamp any active notifies against
+	const float PlayLength = AnimationObject->GetDataModel()->GetPlayLength();
+
+	// Update the end time of any notifies that are currently active
+	for (FRecordedAnimNotify& RecordedNotify : RecordingAnimNotifies)
+	{
+		RecordedNotify.AnimNotifyEndTime = PlayLength;
+		RecordedAnimNotifies.Add(RecordedNotify);
+	}
+
+	RecordingAnimNotifies.Reset();
+
+	for (FRecordedAnimNotify& RecordedNotify : RecordedAnimNotifies)
+	{
+		RecordedNotify.NewNotifyEvent.SetTime(RecordedNotify.AnimNotifyStartTime);
+		RecordedNotify.NewNotifyEvent.TriggerTimeOffset = 0.0f;
+		RecordedNotify.NewNotifyEvent.EndTriggerTimeOffset = 0.0f;
+
+		if (RecordedNotify.OriginalNotifyEvent->GetDuration() > 0.f)
+		{
+			RecordedNotify.NewNotifyEvent.SetDuration(RecordedNotify.AnimNotifyEndTime - RecordedNotify.AnimNotifyStartTime);
+		}
+
+		// see if we need to create a new notify
+		if (RecordedNotify.OriginalNotifyEvent->Notify)
 		{
-			if (Event.TrackIndex >= AnimationObject->AnimNotifyTracks.Num())
+			UAnimNotify** FoundNotify = UniqueNotifies.Find(RecordedNotify.OriginalNotifyEvent->Notify);
+			if (FoundNotify == nullptr)
+			{
+				RecordedNotify.NewNotifyEvent.Notify = Cast<UAnimNotify>(StaticDuplicateObject(RecordedNotify.NewNotifyEvent.Notify, AnimationObject));
+				UniqueNotifies.Add(RecordedNotify.OriginalNotifyEvent->Notify, RecordedNotify.NewNotifyEvent.Notify);
+			}
+			else
 			{
-				AnimationObject->AnimNotifyTracks.SetNum(Event.TrackIndex + 1);
+				RecordedNotify.NewNotifyEvent.Notify = *FoundNotify;
+			}
+		}
 
-				// remake track names to create a nice sequence
-				const int32 TrackNum = AnimationObject->AnimNotifyTracks.Num();
-				for (int32 TrackIndex = 0; TrackIndex < TrackNum; ++TrackIndex)
-				{
-					FAnimNotifyTrack& Track = AnimationObject->AnimNotifyTracks[TrackIndex];
-					Track.TrackName = *FString::FromInt(TrackIndex + 1);
-				}
+		// see if we need to create a new notify state
+		if (RecordedNotify.OriginalNotifyEvent->NotifyStateClass)
+		{
+			UAnimNotifyState** FoundNotifyState = UniqueNotifyStates.Find(RecordedNotify.OriginalNotifyEvent->NotifyStateClass);
+			if (FoundNotifyState == nullptr)
+			{
+				RecordedNotify.NewNotifyEvent.NotifyStateClass = Cast<UAnimNotifyState>(StaticDuplicateObject(RecordedNotify.NewNotifyEvent.NotifyStateClass, AnimationObject));
+				UniqueNotifyStates.Add(RecordedNotify.OriginalNotifyEvent->NotifyStateClass, RecordedNotify.NewNotifyEvent.NotifyStateClass);
+			}
+			else
+			{
+				RecordedNotify.NewNotifyEvent.NotifyStateClass = *FoundNotifyState;
 			}
 		}
 
-		// now build tracks
-		for (int32 EventIndex = 0; EventIndex < AnimationObject->Notifies.Num(); ++EventIndex)
+		AnimationObject->Notifies.Add(RecordedNotify.NewNotifyEvent);
+	}
+
+	// build notify tracks - first find how many tracks we want
+	for (FAnimNotifyEvent& Event : AnimationObject->Notifies)
+	{
+		if (Event.TrackIndex >= AnimationObject->AnimNotifyTracks.Num())
 		{
-			FAnimNotifyEvent& Event = AnimationObject->Notifies[EventIndex];
-			AnimationObject->AnimNotifyTracks[Event.TrackIndex].Notifies.Add(&AnimationObject->Notifies[EventIndex]);
+			AnimationObject->AnimNotifyTracks.SetNum(Event.TrackIndex + 1);
+
+			// remake track names to create a nice sequence
+			const int32 TrackNum = AnimationObject->AnimNotifyTracks.Num();
+			for (int32 TrackIndex = 0; TrackIndex < TrackNum; ++TrackIndex)
+			{
+				FAnimNotifyTrack& Track = AnimationObject->AnimNotifyTracks[TrackIndex];
+				Track.TrackName = *FString::FromInt(TrackIndex + 1);
+			}
 		}
 	}
+
+	// now build tracks
+	for (int32 EventIndex = 0; EventIndex < AnimationObject->Notifies.Num(); ++EventIndex)
+	{
+		FAnimNotifyEvent& Event = AnimationObject->Notifies[EventIndex];
+		AnimationObject->AnimNotifyTracks[Event.TrackIndex].Notifies.Add(&AnimationObject->Notifies[EventIndex]);
+	}
 }
 
 bool FAnimationRecorder::ShouldSkipName(const FName& InName) const
@@ -961,9 +1022,9 @@ void FAnimationRecorder::RecordNotifies(USkeletalMeshComponent* Component, const
 	if (ensure(AnimationObject))
 	{
 		// flag notifies as possibly unused this frame
-		for (auto& ActiveNotify : ActiveNotifies)
+		for (FRecordedAnimNotify& RecordingNotify : RecordingAnimNotifies)
 		{
-			ActiveNotify.Value = false;
+			RecordingNotify.bWasActive = false;
 		}
 
 		int32 AddedThisFrame = 0;
@@ -971,77 +1032,48 @@ void FAnimationRecorder::RecordNotifies(USkeletalMeshComponent* Component, const
 		{
 			if(const FAnimNotifyEvent* NotifyEvent = NotifyEventRef.GetNotify())
 			{
-				// we don't want to insert notifies with duration more than once
-				if(NotifyEvent->GetDuration() > 0.0f)
+				if (NotifyEvent->GetDuration() <= 0.f)
+				{
+					// If the duration of this notify <= 0.0 we just store it as recorded
+					RecordedAnimNotifies.Emplace(*NotifyEvent, NotifyEvent, RecordTime, RecordTime);
+				}
+				else
 				{
-					// if this event is active already then don't add it
+					// Check to see if the notify is already active
 					bool bAlreadyActive = false;
-					for (auto& ActiveNotify : ActiveNotifies)
+					for (FRecordedAnimNotify& RecordingNotify : RecordingAnimNotifies)
 					{
-						if(NotifyEvent == ActiveNotify.Key)
+						if (NotifyEvent == RecordingNotify.OriginalNotifyEvent)
 						{
-							// flag as active
-							ActiveNotify.Value = true;
 							bAlreadyActive = true;
+							RecordingNotify.bWasActive = true;
 							break;
 						}
 					}
 
-					// already active, so skip adding
-					if(bAlreadyActive)
-					{
-						continue;
-					}
-					else
-					{
-						// add a new active notify with duration
-						ActiveNotifies.Emplace(NotifyEvent, true);
-					}
-				}
-
-				// make a new notify from this event & set the current time
-				FAnimNotifyEvent NewEvent = *NotifyEvent;
-				NewEvent.SetTime(RecordTime);
-				NewEvent.TriggerTimeOffset = 0.0f;
-				NewEvent.EndTriggerTimeOffset = 0.0f;
-
-				// see if we need to create a new notify
-				if(NotifyEvent->Notify)
-				{
-					UAnimNotify** FoundNotify = UniqueNotifies.Find(NotifyEvent->Notify);
-					if(FoundNotify == nullptr)
-					{
-						NewEvent.Notify = Cast<UAnimNotify>(StaticDuplicateObject(NewEvent.Notify, AnimationObject));
-						UniqueNotifies.Add(NotifyEvent->Notify, NewEvent.Notify);
-					}
-					else
-					{
-						NewEvent.Notify = *FoundNotify;
-					}
-				}
-
-				// see if we need to create a new notify state
-				if (NotifyEvent->NotifyStateClass)
-				{
-					UAnimNotifyState** FoundNotifyState = UniqueNotifyStates.Find(NotifyEvent->NotifyStateClass);
-					if (FoundNotifyState == nullptr)
+					// Add this 'new' notify
+					if (!bAlreadyActive)
 					{
-						NewEvent.NotifyStateClass = Cast<UAnimNotifyState>(StaticDuplicateObject(NewEvent.NotifyStateClass, AnimationObject));
-						UniqueNotifyStates.Add(NotifyEvent->NotifyStateClass, NewEvent.NotifyStateClass);
-					}
-					else
-					{
-						NewEvent.NotifyStateClass = *FoundNotifyState;
+						RecordingAnimNotifies.Emplace(*NotifyEvent, NotifyEvent, RecordTime, RecordTime);
 					}
 				}
 
-				RecordedNotifyEvents.Add(NewEvent);
 				AddedThisFrame++;
 			}
 		}
 
-		// remove all notifies that didnt get added this time
-		ActiveNotifies.RemoveAll([](TPair<const FAnimNotifyEvent*, bool>& ActiveNotify){ return !ActiveNotify.Value; });
+		for (int32 RecordingAnimNotifyIndex = 0; RecordingAnimNotifyIndex < RecordingAnimNotifies.Num(); ++RecordingAnimNotifyIndex)
+		{
+			FRecordedAnimNotify& RecordingNotify = RecordingAnimNotifies[RecordingAnimNotifyIndex];
+
+			// If a notify wasnt active this frame, commit it as 'recorded'
+			if (!RecordingNotify.bWasActive)
+			{
+				RecordingNotify.AnimNotifyEndTime = RecordTime;
+				RecordedAnimNotifies.Add(RecordingNotify);
+				RecordingAnimNotifies.RemoveAtSwap(RecordingAnimNotifyIndex);
+			}
+		}
 
 		UE_LOG(LogAnimation, Log, TEXT("Added notifies : %d"), AddedThisFrame);
 	}
diff --git a/Engine/Source/Editor/SequenceRecorder/Public/AnimationRecorder.h b/Engine/Source/Editor/SequenceRecorder/Public/AnimationRecorder.h
--- a/Engine/Source/Editor/SequenceRecorder/Public/AnimationRecorder.h
+++ b/Engine/Source/Editor/SequenceRecorder/Public/AnimationRecorder.h
@@ -11,6 +11,7 @@
 #include "Animation/AnimNotifyQueue.h"
 #include "Serializers/MovieSceneAnimationSerialization.h"
 #include "Misc/QualifiedFrameTime.h"
+#include "Animation/AnimTypes.h"
 #include "AnimationRecorder.generated.h"
 
 class UAnimBoneCompressionSettings;
@@ -70,17 +71,41 @@ private:
 	FTransform InitialRootTransform;
 	int32 SkeletonRootIndex;
 
-	/** Array of currently active notifies that have duration */
-	TArray<TPair<const FAnimNotifyEvent*, bool>> ActiveNotifies;
-
 	/** Unique notifies added to this sequence during recording */
 	TMap<UAnimNotify*, UAnimNotify*> UniqueNotifies;
 
 	/** Unique notify states added to this sequence during recording */
 	TMap<UAnimNotifyState*, UAnimNotifyState*> UniqueNotifyStates;
 
+	struct FRecordedAnimNotify
+	{
+		FRecordedAnimNotify(const FAnimNotifyEvent& InNewNotifyEvent, const FAnimNotifyEvent* InOriginalNotifyEvent, float InAnimNotifyStartTime, float InAnimNotifyEndTime)
+			: NewNotifyEvent(InNewNotifyEvent)
+			, OriginalNotifyEvent(InOriginalNotifyEvent)
+			, AnimNotifyStartTime(InAnimNotifyStartTime)
+			, AnimNotifyEndTime(InAnimNotifyEndTime)
+			, bWasActive(true)
+		{}
+
+		/** Notify which will be added to this sequence */
+		FAnimNotifyEvent NewNotifyEvent;
+
+		/** Notify which was called on the sequence being recorded */
+		const FAnimNotifyEvent* OriginalNotifyEvent;
+
+		/** The time in the recorded animation at which the recorded notify started and ended */
+		float AnimNotifyStartTime;
+		float AnimNotifyEndTime;
+
+		/** Whether this notify was active this frame */
+		bool bWasActive;
+	};
+
 	/** Notify events recorded at any point, processed and inserted into animation when recording has finished */
-	TArray<FAnimNotifyEvent> RecordedNotifyEvents;
+	TArray<FRecordedAnimNotify> RecordedAnimNotifies;
+
+	/** Currently recording notify events that have duration */
+	TArray<FRecordedAnimNotify> RecordingAnimNotifies;
 
 	static float DefaultSampleRate;
 
