Add a TCRate property to take track and skeleton bone data.  This will allow pipelines to know what the original TC rate was at recording time.

diff --git a/Engine/Plugins/VirtualProduction/Takes/Source/TakeMovieScene/Private/MovieSceneTakeSection.cpp b/Engine/Plugins/VirtualProduction/Takes/Source/TakeMovieScene/Private/MovieSceneTakeSection.cpp
--- a/Engine/Plugins/VirtualProduction/Takes/Source/TakeMovieScene/Private/MovieSceneTakeSection.cpp
+++ b/Engine/Plugins/VirtualProduction/Takes/Source/TakeMovieScene/Private/MovieSceneTakeSection.cpp
@@ -32,11 +32,14 @@ struct FTakeSectionEditorData
 
 		MetaData[5].SortOrder = 5;
 		MetaData[5].bCanCollapseToTrack = false;
+
+		MetaData[6].SortOrder = 6;
+		MetaData[6].bCanCollapseToTrack = false;
 	}
 
-	FMovieSceneChannelMetaData      MetaData[6];
+	FMovieSceneChannelMetaData      MetaData[7];
 	TMovieSceneExternalValue<int32> ExternalValues[4];
-	TMovieSceneExternalValue<float> ExternalFloatValues[1];
+	TMovieSceneExternalValue<float> ExternalFloatValues[2];
 	TMovieSceneExternalValue<FString> ExternalStringValues[1];
 };
 
@@ -86,6 +89,7 @@ void UMovieSceneTakeSection::ReconstructChannelProxy()
 	FString SecondsName = GetDefault<UMovieSceneTakeSettings>()->SecondsName;
 	FString FramesName = GetDefault<UMovieSceneTakeSettings>()->FramesName;
 	FString SubFramesName = GetDefault<UMovieSceneTakeSettings>()->SubFramesName;
+	FString RateName = GetDefault<UMovieSceneTakeSettings>()->RateName;
 	FString SlateName = GetDefault<UMovieSceneTakeSettings>()->SlateName;
 
 	EditorData.MetaData[0].SetIdentifiers(FName(*HoursName), FText::FromString(HoursName));
@@ -93,14 +97,16 @@ void UMovieSceneTakeSection::ReconstructChannelProxy()
 	EditorData.MetaData[2].SetIdentifiers(FName(*SecondsName), FText::FromString(SecondsName));
 	EditorData.MetaData[3].SetIdentifiers(FName(*FramesName), FText::FromString(FramesName));
 	EditorData.MetaData[4].SetIdentifiers(FName(*SubFramesName), FText::FromString(SubFramesName));
-	EditorData.MetaData[5].SetIdentifiers(FName(*SlateName), FText::FromString(SlateName));
+	EditorData.MetaData[5].SetIdentifiers(FName(*RateName), FText::FromString(RateName));
+	EditorData.MetaData[6].SetIdentifiers(FName(*SlateName), FText::FromString(SlateName));
 
 	Channels.Add(HoursCurve,     EditorData.MetaData[0], EditorData.ExternalValues[0]);
 	Channels.Add(MinutesCurve,   EditorData.MetaData[1], EditorData.ExternalValues[1]);
 	Channels.Add(SecondsCurve,   EditorData.MetaData[2], EditorData.ExternalValues[2]);
 	Channels.Add(FramesCurve,    EditorData.MetaData[3], EditorData.ExternalValues[3]);
 	Channels.Add(SubFramesCurve, EditorData.MetaData[4], EditorData.ExternalFloatValues[0]);
-	Channels.Add(Slate, EditorData.MetaData[5], EditorData.ExternalStringValues[0]);
+	Channels.Add(RateCurve,      EditorData.MetaData[5], EditorData.ExternalFloatValues[1]);
+	Channels.Add(Slate,          EditorData.MetaData[6], EditorData.ExternalStringValues[0]);
 
 #else
 
@@ -109,6 +115,7 @@ void UMovieSceneTakeSection::ReconstructChannelProxy()
 	Channels.Add(SecondsCurve);
 	Channels.Add(FramesCurve);
 	Channels.Add(SubFramesCurve);
+	Channels.Add(RateCurve);
 	Channels.Add(Slate);
 
 #endif
diff --git a/Engine/Plugins/VirtualProduction/Takes/Source/TakeMovieScene/Private/MovieSceneTakeSettings.cpp b/Engine/Plugins/VirtualProduction/Takes/Source/TakeMovieScene/Private/MovieSceneTakeSettings.cpp
--- a/Engine/Plugins/VirtualProduction/Takes/Source/TakeMovieScene/Private/MovieSceneTakeSettings.cpp
+++ b/Engine/Plugins/VirtualProduction/Takes/Source/TakeMovieScene/Private/MovieSceneTakeSettings.cpp
@@ -12,6 +12,7 @@ UMovieSceneTakeSettings::UMovieSceneTakeSettings()
 	, SecondsName(TEXT("TCSecond"))
 	, FramesName(TEXT("TCFrame"))
 	, SubFramesName(TEXT("TCSubframe"))
+	, RateName(TEXT("TCRate"))
 	, SlateName(TEXT("TCSlate"))
 {
 }
diff --git a/Engine/Plugins/VirtualProduction/Takes/Source/TakeMovieScene/Public/MovieSceneTakeSection.h b/Engine/Plugins/VirtualProduction/Takes/Source/TakeMovieScene/Public/MovieSceneTakeSection.h
--- a/Engine/Plugins/VirtualProduction/Takes/Source/TakeMovieScene/Public/MovieSceneTakeSection.h
+++ b/Engine/Plugins/VirtualProduction/Takes/Source/TakeMovieScene/Public/MovieSceneTakeSection.h
@@ -58,6 +58,10 @@ public:
 	UPROPERTY()
 	FMovieSceneFloatChannel SubFramesCurve;
 
+	/** Timecode rate curve data */
+	UPROPERTY()
+	FMovieSceneFloatChannel RateCurve;
+
 	/** Slate data */
 	UPROPERTY()
 	FMovieSceneStringChannel Slate;
diff --git a/Engine/Plugins/VirtualProduction/Takes/Source/TakeMovieScene/Public/MovieSceneTakeSettings.h b/Engine/Plugins/VirtualProduction/Takes/Source/TakeMovieScene/Public/MovieSceneTakeSettings.h
--- a/Engine/Plugins/VirtualProduction/Takes/Source/TakeMovieScene/Public/MovieSceneTakeSettings.h
+++ b/Engine/Plugins/VirtualProduction/Takes/Source/TakeMovieScene/Public/MovieSceneTakeSettings.h
@@ -40,6 +40,10 @@ public:
 	UPROPERTY(config, EditAnywhere, BlueprintReadWrite, Category = "Take Recorder")
 	FString SubFramesName;
 
+	/** Rate Name */
+	UPROPERTY(config, EditAnywhere, BlueprintReadWrite, Category = "Take Recorder")
+	FString RateName;
+
 	/** Slate Name */
 	UPROPERTY(config, EditAnywhere, BlueprintReadWrite, Category = "Take Recorder")
 	FString SlateName;
diff --git a/Engine/Plugins/VirtualProduction/Takes/Source/TakeRecorderSources/Private/TakeRecorderActorSource.cpp b/Engine/Plugins/VirtualProduction/Takes/Source/TakeRecorderSources/Private/TakeRecorderActorSource.cpp
--- a/Engine/Plugins/VirtualProduction/Takes/Source/TakeRecorderSources/Private/TakeRecorderActorSource.cpp
+++ b/Engine/Plugins/VirtualProduction/Takes/Source/TakeRecorderSources/Private/TakeRecorderActorSource.cpp
@@ -1,6 +1,7 @@
 // Copyright Epic Games, Inc. All Rights Reserved.
 
 #include "TakeRecorderActorSource.h"
+#include "AnimationRecorder.h"
 #include "Styling/SlateIconFinder.h"
 #include "ClassIconFinder.h"
 #include "MovieScene.h"
@@ -707,6 +708,7 @@ void UTakeRecorderActorSource::ProcessRecordedTimes(ULevelSequence* InSequence)
 		FFrameRate TickResolution = MovieScene->GetTickResolution();
 		FFrameRate DisplayRate = MovieScene->GetDisplayRate();
 
+		FFrameRate TCRate = TickResolution;
 		for (const TPair<FQualifiedFrameTime, FQualifiedFrameTime>& RecordedTimePair : RecordedTimes)
 		{
 			FFrameNumber FrameNumber = RecordedTimePair.Key.Time.FrameNumber;
@@ -716,7 +718,7 @@ void UTakeRecorderActorSource::ProcessRecordedTimes(ULevelSequence* InSequence)
 			}
 
 			FTimecode Timecode = RecordedTimePair.Value.ToTimecode();
-		
+			TCRate = RecordedTimePair.Value.Rate;
 			Hours.Add(Timecode.Hours);
 			Minutes.Add(Timecode.Minutes);
 			Seconds.Add(Timecode.Seconds);
@@ -732,7 +734,7 @@ void UTakeRecorderActorSource::ProcessRecordedTimes(ULevelSequence* InSequence)
 			else
 			{
 				FFrameTime FrameTime = FFrameRate::TransformTime(RecordedTimePair.Key.Time, TickResolution, DisplayRate);
-				FQualifiedFrameTime FrameTimeAsTimeCodeRate(FrameTime, RecordedTimePair.Value.Rate);
+				FQualifiedFrameTime FrameTimeAsTimeCodeRate(FrameTime, TCRate);
 
 				SubFrame.Value = FrameTimeAsTimeCodeRate.Time.GetSubFrame();
 			}
@@ -755,6 +757,7 @@ void UTakeRecorderActorSource::ProcessRecordedTimes(ULevelSequence* InSequence)
 		TakeSection->SecondsCurve.Set(Times, Seconds);
 		TakeSection->FramesCurve.Set(Times, Frames);
 		TakeSection->SubFramesCurve.Set(Times, SubFrames);
+		TakeSection->RateCurve.SetDefault(TCRate.AsDecimal());
 	}
 
 	// Since the take section was created post recording here in this
@@ -876,23 +879,37 @@ void UTakeRecorderActorSource::FinalizeRecording()
 	ParentSource = nullptr;
 }
 
-void UTakeRecorderActorSource::PostProcessTrackRecorders(ULevelSequence* InSequence)
+namespace UE::TakeRecorderActorSource::Private
+{
+FProcessRecordedTimeParams GetTimecodeRecordingParameters()
 {
-	FTakeRecorderParameters Parameters;
-	Parameters.User = GetDefault<UTakeRecorderUserSettings>()->Settings;
-	Parameters.Project = GetDefault<UTakeRecorderProjectSettings>()->Settings;
-
 	FString HoursName = GetDefault<UMovieSceneTakeSettings>()->HoursName;
 	FString MinutesName = GetDefault<UMovieSceneTakeSettings>()->MinutesName;
 	FString SecondsName = GetDefault<UMovieSceneTakeSettings>()->SecondsName;
 	FString FramesName = GetDefault<UMovieSceneTakeSettings>()->FramesName;
 	FString SubFramesName = GetDefault<UMovieSceneTakeSettings>()->SubFramesName;
 	FString SlateName = GetDefault<UMovieSceneTakeSettings>()->SlateName;
-				
-	FString Slate;
+
+	return FProcessRecordedTimeParams {
+		.HoursName = HoursName,
+		.MinutesName = MinutesName,
+		.SecondsName = SecondsName,
+		.FramesName = FramesName,
+		.SubFramesName = SubFramesName,
+		.SlateName = SlateName
+	};
+}
+}
+void UTakeRecorderActorSource::PostProcessTrackRecorders(ULevelSequence* InSequence)
+{
+	FTakeRecorderParameters Parameters;
+	Parameters.User = GetDefault<UTakeRecorderUserSettings>()->Settings;
+	Parameters.Project = GetDefault<UTakeRecorderProjectSettings>()->Settings;
+
+	FProcessRecordedTimeParams RecordedTimeParams = UE::TakeRecorderActorSource::Private::GetTimecodeRecordingParameters();
 	if (UTakeMetaData* TakeMetaData = InSequence->FindMetaData<UTakeMetaData>())
 	{
-		Slate = FString::Printf(TEXT("%s_%d"), *TakeMetaData->GetSlate(), TakeMetaData->GetTakeNumber());
+		RecordedTimeParams.Slate = FString::Printf(TEXT("%s_%d"), *TakeMetaData->GetSlate(), TakeMetaData->GetTakeNumber());
 	}
 
 	// We want to look at all Animation Track recorders and remove root motion if the transform
@@ -943,7 +960,7 @@ void UTakeRecorderActorSource::PostProcessTrackRecorders(ULevelSequence* InSeque
 			
 			if (Parameters.Project.bRecordTimecode)
 			{
-				AnimationTrackRecorder->ProcessRecordedTimes(HoursName, MinutesName, SecondsName, FramesName, SubFramesName, SlateName, Slate);
+				AnimationTrackRecorder->ProcessRecordedTimes(RecordedTimeParams);
 			}
 		}
 	}
diff --git a/Engine/Plugins/VirtualProduction/Takes/Source/TakeSequencer/Private/TakeTrackEditor.cpp b/Engine/Plugins/VirtualProduction/Takes/Source/TakeSequencer/Private/TakeTrackEditor.cpp
--- a/Engine/Plugins/VirtualProduction/Takes/Source/TakeSequencer/Private/TakeTrackEditor.cpp
+++ b/Engine/Plugins/VirtualProduction/Takes/Source/TakeSequencer/Private/TakeTrackEditor.cpp
@@ -71,7 +71,8 @@ bool FTakeTrackEditor::ImportAnimatedProperty(const FString& InPropertyName, con
 		InPropertyName != GetDefault<UMovieSceneTakeSettings>()->MinutesName && 
 		InPropertyName != GetDefault<UMovieSceneTakeSettings>()->SecondsName && 
 		InPropertyName != GetDefault<UMovieSceneTakeSettings>()->FramesName && 
-		InPropertyName != GetDefault<UMovieSceneTakeSettings>()->SubFramesName)
+		InPropertyName != GetDefault<UMovieSceneTakeSettings>()->SubFramesName &&
+		InPropertyName != GetDefault<UMovieSceneTakeSettings>()->RateName)
 	{
 		return false;
 	}
@@ -133,6 +134,10 @@ bool FTakeTrackEditor::ImportAnimatedProperty(const FString& InPropertyName, con
 	{
 		TakeSection->SubFramesCurve.Set(KeyTimes, KeyFloatValues);
 	}
+	else if (InPropertyName == GetDefault<UMovieSceneTakeSettings>()->RateName)
+	{
+		TakeSection->RateCurve.Set(KeyTimes, KeyFloatValues);
+	}
 
 	TOptional<TRange<FFrameNumber> > AutoSizeRange = TakeSection->GetAutoSizeRange();
 	if (AutoSizeRange.IsSet())
diff --git a/Engine/Plugins/VirtualProduction/Takes/Source/TakeTrackRecorders/Private/MovieSceneAnimationTrackRecorder.cpp b/Engine/Plugins/VirtualProduction/Takes/Source/TakeTrackRecorders/Private/MovieSceneAnimationTrackRecorder.cpp
--- a/Engine/Plugins/VirtualProduction/Takes/Source/TakeTrackRecorders/Private/MovieSceneAnimationTrackRecorder.cpp
+++ b/Engine/Plugins/VirtualProduction/Takes/Source/TakeTrackRecorders/Private/MovieSceneAnimationTrackRecorder.cpp
@@ -2,6 +2,7 @@
 
 #include "TrackRecorders/MovieSceneAnimationTrackRecorder.h"
 #include "Animation/Skeleton.h"
+#include "AnimationRecorder.h"
 #include "AssetRegistry/AssetData.h"
 #include "Engine/SkeletalMesh.h"
 #include "TrackRecorders/MovieSceneAnimationTrackRecorderSettings.h"
@@ -347,10 +348,23 @@ void UMovieSceneAnimationTrackRecorder::RemoveRootMotion()
 }
 
 void UMovieSceneAnimationTrackRecorder::ProcessRecordedTimes(const FString& HoursName, const FString& MinutesName, const FString& SecondsName, const FString& FramesName, const FString& SubFramesName, const FString& SlateName, const FString& Slate)
+{
+	FProcessRecordedTimeParams Params{
+		 .HoursName = HoursName,
+		 .MinutesName = MinutesName,
+		 .SecondsName = SecondsName,
+		 .FramesName = FramesName,
+		 .SubFramesName = SubFramesName,
+		 .SlateName = SlateName,
+		 .Slate = Slate,
+	};
+	ProcessRecordedTimes(Params);
+}
+void UMovieSceneAnimationTrackRecorder::ProcessRecordedTimes(const FProcessRecordedTimeParams& InParams)
 {
 	UMovieSceneAnimationTrackRecorderSettings* AnimSettings = CastChecked<UMovieSceneAnimationTrackRecorderSettings>(Settings.Get());
 
-	AnimationRecorder.ProcessRecordedTimes(AnimSequence.Get(), SkeletalMeshComponent.Get(), HoursName, MinutesName, SecondsName, FramesName, SubFramesName, SlateName, Slate, AnimSettings->TimecodeBoneMethod);
+	AnimationRecorder.ProcessRecordedTimes(AnimSequence.Get(), SkeletalMeshComponent.Get(), AnimSettings->TimecodeBoneMethod, InParams);
 }
 
 bool UMovieSceneAnimationTrackRecorder::LoadRecordedFile(const FString& FileName, UMovieScene *InMovieScene, TMap<FGuid, AActor*>& ActorGuidToActorMap,  TFunction<void()> InCompletionCallback)
diff --git a/Engine/Plugins/VirtualProduction/Takes/Source/TakeTrackRecorders/Public/TrackRecorders/MovieSceneAnimationTrackRecorder.h b/Engine/Plugins/VirtualProduction/Takes/Source/TakeTrackRecorders/Public/TrackRecorders/MovieSceneAnimationTrackRecorder.h
--- a/Engine/Plugins/VirtualProduction/Takes/Source/TakeTrackRecorders/Public/TrackRecorders/MovieSceneAnimationTrackRecorder.h
+++ b/Engine/Plugins/VirtualProduction/Takes/Source/TakeTrackRecorders/Public/TrackRecorders/MovieSceneAnimationTrackRecorder.h
@@ -59,8 +59,12 @@ protected:
 public:
 	bool RootWasRemoved() const { return bRootWasRemoved; }
 	void RemoveRootMotion();
+	UE_DEPRECATED(5.5, "Use the ProcessRecordedTimes method that takes a FProcessRecordedTimeParams struct.")
 	void ProcessRecordedTimes(const FString& HoursName, const FString& MinutesName, const FString& SecondsName, const FString& FramesName, const FString& SubFramesName, const FString& SlateName, const FString& Slate);
 
+	/** Apply time code data to the animation track. */
+	void ProcessRecordedTimes(const FProcessRecordedTimeParams& InParams);
+
 	UAnimSequence* GetAnimSequence() const { return AnimSequence.Get(); }
 	USkeletalMesh* GetSkeletalMesh() const { return SkeletalMesh.Get(); }
 	USkeletalMeshComponent* GetSkeletalMeshComponent() const { return SkeletalMeshComponent.Get(); }
diff --git a/Engine/Source/Editor/SequenceRecorder/Private/AnimationRecorder.cpp b/Engine/Source/Editor/SequenceRecorder/Private/AnimationRecorder.cpp
--- a/Engine/Source/Editor/SequenceRecorder/Private/AnimationRecorder.cpp
+++ b/Engine/Source/Editor/SequenceRecorder/Private/AnimationRecorder.cpp
@@ -1,6 +1,7 @@
 // Copyright Epic Games, Inc. All Rights Reserved.
 
 #include "AnimationRecorder.h"
+#include "Misc/QualifiedFrameTime.h"
 #include "Widgets/DeclarativeSyntaxSupport.h"
 #include "Animation/AnimSequence.h"
 #include "Components/SkeletalMeshComponent.h"
@@ -642,8 +643,21 @@ UAnimSequence* FAnimationRecorder::StopRecord(bool bShowMessage)
 
 	return nullptr;
 }
-
 void FAnimationRecorder::ProcessRecordedTimes(UAnimSequence* AnimSequence, USkeletalMeshComponent* SkeletalMeshComponent, const FString& HoursName, const FString& MinutesName, const FString& SecondsName, const FString& FramesName, const FString& SubFramesName, const FString& SlateName, const FString& Slate, const FTimecodeBoneMethod& TimecodeBoneMethod)
+{
+	FProcessRecordedTimeParams Params{
+		 .HoursName = HoursName,
+		 .MinutesName = MinutesName,
+		 .SecondsName = SecondsName,
+		 .FramesName = FramesName,
+		 .SubFramesName = SubFramesName,
+		 .SlateName = SlateName,
+		 .Slate = Slate
+	};
+	ProcessRecordedTimes(AnimSequence, SkeletalMeshComponent, TimecodeBoneMethod, Params);
+}
+
+void FAnimationRecorder::ProcessRecordedTimes(UAnimSequence* AnimSequence, USkeletalMeshComponent* SkeletalMeshComponent, const FTimecodeBoneMethod& TimecodeBoneMethod, const FProcessRecordedTimeParams& InParams)
 {
 	if (!AnimSequence || !SkeletalMeshComponent)
 	{
@@ -667,13 +681,14 @@ void FAnimationRecorder::ProcessRecordedTimes(UAnimSequence* AnimSequence, USkel
 	SubFrames.Reserve(RecordedTimes.Num());
 	Times.Reserve(RecordedTimes.Num());
 
+	FFrameRate TCRate = RecordingRate;
 	for (int32 KeyIndex = 0; KeyIndex < NumKeys; ++KeyIndex)
 	{
 		const float TimeToRecord = RecordingRate.AsSeconds(KeyIndex);
 
 		FQualifiedFrameTime RecordedTime = RecordedTimes[KeyIndex];
 		FTimecode Timecode = FTimecode::FromFrameNumber(RecordedTime.Time.FrameNumber, RecordedTime.Rate);
-		
+		TCRate = RecordedTime.Rate;
 		Hours.Add(Timecode.Hours);
 		Minutes.Add(Timecode.Minutes);
 		Seconds.Add(Timecode.Seconds);
@@ -706,9 +721,18 @@ void FAnimationRecorder::ProcessRecordedTimes(UAnimSequence* AnimSequence, USkel
 	}
 
 	// String is not animatable, just add 1 slate at the first key time
-	TArray<FString> Slates(&Slate, 1);
+	TArray<FString> Slates(&InParams.Slate, 1);
 	TArray<float> SlateTimes(&Times[0], 1);
 
+	float RecordingRateAsFloat = TCRate.AsDecimal();
+	TOptional<TArray<float>> Rates;
+	TOptional<TArray<float>> RateTimes;
+	if (InParams.RateName)
+	{
+		Rates = TArray<float>(&RecordingRateAsFloat, 1);
+		RateTimes = TArray<float>(&Times[0],1);
+	}
+
 	IAnimationDataController& Controller = AnimSequence->GetController();
 	IAnimationDataController::FScopedBracket ScopedBracket(Controller, LOCTEXT("AddTimeCodeAttributesBracket", "Adding Time Code attributes"));
 
@@ -750,7 +774,7 @@ void FAnimationRecorder::ProcessRecordedTimes(UAnimSequence* AnimSequence, USkel
 		{
 			// add tracks for the bone existing
 			FName BoneTreeName = AnimSkeleton->GetReferenceSkeleton().GetBoneName(BoneTreeIndex);
-			
+
 			const bool bUseThisBone = 
 				TimecodeBoneMethod.BoneMode == ETimecodeBoneMode::All ||
 				(TimecodeBoneMethod.BoneMode == ETimecodeBoneMode::Root && BoneIndex == 0) ||
@@ -762,14 +786,19 @@ void FAnimationRecorder::ProcessRecordedTimes(UAnimSequence* AnimSequence, USkel
 				continue;
 			}
 
-			UE::Anim::AddTypedCustomAttribute<FIntegerAnimationAttribute, int32>(FName(*HoursName), BoneTreeName, AnimSequence, MakeArrayView(Times), MakeArrayView(Hours));
-			UE::Anim::AddTypedCustomAttribute<FIntegerAnimationAttribute, int32>(FName(*MinutesName), BoneTreeName, AnimSequence, MakeArrayView(Times), MakeArrayView(Minutes));
-			UE::Anim::AddTypedCustomAttribute<FIntegerAnimationAttribute, int32>(FName(*SecondsName), BoneTreeName, AnimSequence, MakeArrayView(Times), MakeArrayView(Seconds));
-			UE::Anim::AddTypedCustomAttribute<FIntegerAnimationAttribute, int32>(FName(*FramesName), BoneTreeName, AnimSequence, MakeArrayView(Times), MakeArrayView(Frames));
+			UE::Anim::AddTypedCustomAttribute<FIntegerAnimationAttribute, int32>(FName(*InParams.HoursName), BoneTreeName, AnimSequence, MakeArrayView(Times), MakeArrayView(Hours));
+			UE::Anim::AddTypedCustomAttribute<FIntegerAnimationAttribute, int32>(FName(*InParams.MinutesName), BoneTreeName, AnimSequence, MakeArrayView(Times), MakeArrayView(Minutes));
+			UE::Anim::AddTypedCustomAttribute<FIntegerAnimationAttribute, int32>(FName(*InParams.SecondsName), BoneTreeName, AnimSequence, MakeArrayView(Times), MakeArrayView(Seconds));
+			UE::Anim::AddTypedCustomAttribute<FIntegerAnimationAttribute, int32>(FName(*InParams.FramesName), BoneTreeName, AnimSequence, MakeArrayView(Times), MakeArrayView(Frames));
 
-			UE::Anim::AddTypedCustomAttribute<FFloatAnimationAttribute, float>(FName(*SubFramesName), BoneTreeName, AnimSequence, MakeArrayView(Times), MakeArrayView(SubFrames));
+			UE::Anim::AddTypedCustomAttribute<FFloatAnimationAttribute, float>(FName(*InParams.SubFramesName), BoneTreeName, AnimSequence, MakeArrayView(Times), MakeArrayView(SubFrames));
+
+			if (Rates && RateTimes)
+			{
+				UE::Anim::AddTypedCustomAttribute<FFloatAnimationAttribute, float>(FName(*(InParams.RateName.GetValue())), BoneTreeName, AnimSequence, MakeArrayView(*RateTimes), MakeArrayView(*Rates));
+			}
+			UE::Anim::AddTypedCustomAttribute<FStringAnimationAttribute, FString>(FName(*InParams.SlateName), BoneTreeName, AnimSequence, MakeArrayView(SlateTimes), MakeArrayView(Slates));
 
-			UE::Anim::AddTypedCustomAttribute<FStringAnimationAttribute, FString>(FName(*SlateName), BoneTreeName, AnimSequence, MakeArrayView(SlateTimes), MakeArrayView(Slates));
 		}
 	}
 }
@@ -998,7 +1027,22 @@ bool FAnimationRecorder::Record(USkeletalMeshComponent* Component, FTransform co
 		}
 
 		TOptional<FQualifiedFrameTime> CurrentTime = FApp::GetCurrentFrameTime();
-		RecordedTimes.Add(CurrentTime.IsSet() ? CurrentTime.GetValue() : FQualifiedFrameTime());
+		if (CurrentTime)
+		{
+			FQualifiedFrameTime TimeToRecord = *CurrentTime;
+			if (TimeToRecord.Time.GetSubFrame() == 0)
+			{
+				FQualifiedFrameTime FrameTimeAsQualified(FrameToAdd, RecordingRate);
+				FFrameTime FrameTimeInTCRate = FrameTimeAsQualified.ConvertTo(TimeToRecord.Rate);
+				FFrameTime NewFrameTimeWithSubFrame(TimeToRecord.Time.GetFrame(), FrameTimeInTCRate.GetSubFrame());
+				TimeToRecord = FQualifiedFrameTime(NewFrameTimeWithSubFrame, CurrentTime->Rate);
+			}
+			RecordedTimes.Add(TimeToRecord);
+		}
+		else
+		{
+			RecordedTimes.Add(FQualifiedFrameTime());
+		}
 
 		if (AnimationSerializer)
 		{
@@ -1238,10 +1282,26 @@ void FAnimRecorderInstance::ProcessRecordedTimes(UAnimSequence* AnimSequence, US
 {
 	if (Recorder.IsValid())
 	{
-		Recorder->ProcessRecordedTimes(AnimSequence, SkeletalMeshComponent, HoursName, MinutesName, SecondsName, FramesName, SubFramesName, SlateName, Slate, TimecodeBoneMethod);
+		FProcessRecordedTimeParams Params{
+		 .HoursName = HoursName,
+		 .MinutesName = MinutesName,
+		 .SecondsName = SecondsName,
+		 .FramesName = FramesName,
+		 .SubFramesName = SubFramesName,
+		 .SlateName = SlateName,
+		 .Slate = Slate
+		};
+		Recorder->ProcessRecordedTimes(AnimSequence, SkeletalMeshComponent, TimecodeBoneMethod, Params);
 	}
 }
 
+void FAnimRecorderInstance::ProcessRecordedTimes(UAnimSequence* AnimSequence, USkeletalMeshComponent* SkeletalMeshComponent, const FTimecodeBoneMethod& TimecodeBoneMethod, const FProcessRecordedTimeParams& InParams)
+{
+	if (Recorder.IsValid())
+	{
+		Recorder->ProcessRecordedTimes(AnimSequence, SkeletalMeshComponent, TimecodeBoneMethod, InParams);
+	}
+}
 
 bool FAnimationRecorderManager::RecordAnimation(USkeletalMeshComponent* Component, const FString& AssetPath, const FString& AssetName, const FAnimationRecordingSettings& Settings)
 {
diff --git a/Engine/Source/Editor/SequenceRecorder/Public/AnimationRecorder.h b/Engine/Source/Editor/SequenceRecorder/Public/AnimationRecorder.h
--- a/Engine/Source/Editor/SequenceRecorder/Public/AnimationRecorder.h
+++ b/Engine/Source/Editor/SequenceRecorder/Public/AnimationRecorder.h
@@ -48,6 +48,20 @@ struct FTimecodeBoneMethod
 	FName BoneName;
 };
 
+struct SEQUENCERECORDER_API FProcessRecordedTimeParams
+{
+	FString HoursName;
+	FString MinutesName;
+	FString SecondsName;
+	FString FramesName;
+	FString SubFramesName;
+
+	// Optionally support writing the Timecode Rate into time data.
+	TOptional<FString> RateName;
+
+	FString SlateName;
+	FString Slate;
+};
 //////////////////////////////////////////////////////////////////////////
 // FAnimationRecorder
 
@@ -142,8 +156,12 @@ public:
 
 	const FTransform& GetInitialRootTransform() const { return InitialRootTransform; }
 
+	UE_DEPRECATED(5.5, "Use the ProcessRecordedTimes method that takes a FProcessRecordedTimeParams struct.")
 	void ProcessRecordedTimes(UAnimSequence* AnimSequence, USkeletalMeshComponent* SkeletalMeshComponent, const FString& HoursName, const FString& MinutesName, const FString& SecondsName, const FString& FramesName, const FString& SubFramesName, const FString& SlateName, const FString& Slate, const FTimecodeBoneMethod& TimecodeBoneMethod);
 
+	/** Process any time data captured and apply it to the bones on the given SkeletalMeshComponent. */
+	void ProcessRecordedTimes(UAnimSequence* AnimSequence, USkeletalMeshComponent* SkeletalMeshComponent, const FTimecodeBoneMethod& TimecodeBoneMethod, const FProcessRecordedTimeParams& TimecodeInfo);
+
 	/** If true, it will record root to include LocalToWorld */
 	uint8 bRecordLocalToWorld :1;
 	/** If true, asset will be saved to disk after recording. If false, asset will remain in mem and can be manually saved. */
@@ -207,8 +225,13 @@ public:
 	bool BeginRecording();
 	void Update(float DeltaTime);
 	void FinishRecording(bool bShowMessage = true);
+
+	UE_DEPRECATED(5.5, "Use the ProcessRecordedTimes method that takes a FProcessRecordedTimeParams struct.")
 	void ProcessRecordedTimes(UAnimSequence* AnimSequence, USkeletalMeshComponent* SkeletalMeshComponent, const FString& HoursName, const FString& MinutesName, const FString& SecondsName, const FString& FramesName, const FString& SubFramesName, const FString& SlateName, const FString& Slate, const FTimecodeBoneMethod& TimecodeBoneMethod);
 
+	/** Process any time data captured and apply it to the bones on the given SkeletalMeshComponent. */
+	void ProcessRecordedTimes(UAnimSequence* AnimSequence, USkeletalMeshComponent* SkeletalMeshComponent, const FTimecodeBoneMethod& TimecodeBoneMethod, const FProcessRecordedTimeParams& TimecodeInfo);
+
 private:
 	void InitInternal(USkeletalMeshComponent* InComponent, const FAnimationRecordingSettings& Settings, FAnimationSerializer *InAnimationSerializer = nullptr);
 
