[trivial] Removed little-used but heavy header from the Core PCH
- also added a couple of possibly unrelated missing includes

diff --git a/Engine/Plugins/VirtualProduction/DMX/DMXProtocol/Source/DMXProtocol/Private/IO/DMXOutputPort.cpp b/Engine/Plugins/VirtualProduction/DMX/DMXProtocol/Source/DMXProtocol/Private/IO/DMXOutputPort.cpp
--- a/Engine/Plugins/VirtualProduction/DMX/DMXProtocol/Source/DMXProtocol/Private/IO/DMXOutputPort.cpp
+++ b/Engine/Plugins/VirtualProduction/DMX/DMXProtocol/Source/DMXProtocol/Private/IO/DMXOutputPort.cpp
@@ -14,8 +14,9 @@
 #include "IO/DMXOutputPortConfig.h"
 #include "IO/DMXPortManager.h"
 #include "IO/DMXRawListener.h"
-#include <limits>
+#include "Misc/ExpressionParserTypes.h"
 #include "Misc/FrameRate.h"
+#include <limits>
 
 
 #define LOCTEXT_NAMESPACE "DMXOutputPort"
diff --git a/Engine/Plugins/VirtualProduction/VirtualCamera/Source/VirtualCamera/Private/FunctionLibraries/VCamBlueprintFunctionLibrary.cpp b/Engine/Plugins/VirtualProduction/VirtualCamera/Source/VirtualCamera/Private/FunctionLibraries/VCamBlueprintFunctionLibrary.cpp
--- a/Engine/Plugins/VirtualProduction/VirtualCamera/Source/VirtualCamera/Private/FunctionLibraries/VCamBlueprintFunctionLibrary.cpp
+++ b/Engine/Plugins/VirtualProduction/VirtualCamera/Source/VirtualCamera/Private/FunctionLibraries/VCamBlueprintFunctionLibrary.cpp
@@ -13,6 +13,7 @@
 #include "GameFramework/Actor.h"
 #include "GameFramework/PlayerController.h"
 #include "LevelSequence.h"
+#include "Misc/ExpressionParserTypes.h"
 #include "MovieScene.h"
 #include "MovieSceneSequencePlayer.h"
 #include "Slate/SceneViewport.h"
diff --git a/Engine/Source/Editor/SequenceRecorder/Private/SequenceRecorderModule.cpp b/Engine/Source/Editor/SequenceRecorder/Private/SequenceRecorderModule.cpp
--- a/Engine/Source/Editor/SequenceRecorder/Private/SequenceRecorderModule.cpp
+++ b/Engine/Source/Editor/SequenceRecorder/Private/SequenceRecorderModule.cpp
@@ -1,53 +1,54 @@
 // Copyright Epic Games, Inc. All Rights Reserved.
 
-#include "CoreMinimal.h"
-#include "Misc/CoreMisc.h"
-#include "Misc/Guid.h"
-#include "Modules/ModuleManager.h"
-#include "UObject/ObjectMacros.h"
-#include "UObject/Class.h"
-#include "Misc/PackageName.h"
-#include "HAL/IConsoleManager.h"
-#include "Engine/Level.h"
-#include "GameFramework/Actor.h"
-#include "Widgets/DeclarativeSyntaxSupport.h"
+#include "ActorRecording.h"
+#include "ActorRecordingDetailsCustomization.h"
+#include "Animation/AnimationRecordingSettings.h"
+#include "Animation/AnimSequence.h"
+#include "AnimationRecorder.h"
 #include "Components/SkeletalMeshComponent.h"
+#include "Containers/ArrayView.h"
+#include "CoreMinimal.h"
 #include "Debug/DebugDrawService.h"
+#include "Engine/Level.h"
+#include "EngineGlobals.h"
 #include "Framework/Application/SlateApplication.h"
-#include "Textures/SlateIcon.h"
+#include "Framework/Docking/LayoutExtender.h"
 #include "Framework/Docking/TabManager.h"
-#include "Styling/AppStyle.h"
-#include "EngineGlobals.h"
+#include "Framework/Notifications/NotificationManager.h"
+#include "GameFramework/Actor.h"
+#include "HAL/IConsoleManager.h"
+#include "ISequenceAudioRecorder.h"
+#include "ISequenceRecorder.h"
+#include "IStructureDetailsView.h"
 #include "LevelEditor.h"
+#include "LevelSequence.h"
+#include "Misc/CoreMisc.h"
+#include "Misc/ExpressionParserTypes.h"
+#include "Misc/Guid.h"
+#include "Misc/PackageName.h"
+#include "Modules/ModuleManager.h"
+#include "MovieSceneTimeHelpers.h"
 #include "PersonaModule.h"
-#include "Animation/AnimationRecordingSettings.h"
-#include "AnimationRecorder.h"
-#include "WorkspaceMenuStructure.h"
-#include "WorkspaceMenuStructureModule.h"
+#include "PropertiesToRecordForClassDetailsCustomization.h"
 #include "PropertyEditorModule.h"
-#include "ActorRecording.h"
-#include "SSequenceRecorder.h"
-#include "Widgets/Docking/SDockTab.h"
-#include "SequenceRecorderCommands.h"
-#include "Containers/ArrayView.h"
-#include "ISequenceAudioRecorder.h"
-#include "ISequenceRecorder.h"
 #include "SequenceRecorder.h"
-#include "SequenceRecorderSettings.h"
-#include "Framework/Notifications/NotificationManager.h"
-#include "Widgets/Notifications/SNotificationList.h"
-#include "ActorRecordingDetailsCustomization.h"
+#include "SequenceRecorderCommands.h"
 #include "SequenceRecorderDetailsCustomization.h"
-#include "PropertiesToRecordForClassDetailsCustomization.h"
+#include "SequenceRecorderSettings.h"
+#include "SequenceRecorderUtils.h"
+#include "SSequenceRecorder.h"
+#include "Styling/AppStyle.h"
+#include "Textures/SlateIcon.h"
 #include "TimecodeBoneMethodCustomization.h"
+#include "UObject/Class.h"
+#include "UObject/ObjectMacros.h"
+#include "Widgets/DeclarativeSyntaxSupport.h"
+#include "Widgets/Docking/SDockTab.h"
+#include "Widgets/Notifications/SNotificationList.h"
 #include "WorkflowOrientedApp/WorkflowTabFactory.h"
-#include "IStructureDetailsView.h"
 #include "WorkflowOrientedApp/WorkflowTabManager.h"
-#include "Framework/Docking/LayoutExtender.h"
-#include "MovieSceneTimeHelpers.h"
-#include "SequenceRecorderUtils.h"
-#include "Animation/AnimSequence.h"
-#include "LevelSequence.h"
+#include "WorkspaceMenuStructure.h"
+#include "WorkspaceMenuStructureModule.h"
 
 #define LOCTEXT_NAMESPACE "SequenceRecorder"
 
diff --git a/Engine/Source/Runtime/Core/Public/CoreSharedPCH.h b/Engine/Source/Runtime/Core/Public/CoreSharedPCH.h
--- a/Engine/Source/Runtime/Core/Public/CoreSharedPCH.h
+++ b/Engine/Source/Runtime/Core/Public/CoreSharedPCH.h
@@ -89,7 +89,6 @@
 #include "Logging/LogVerbosity.h"
 #include "Logging/TokenizedMessage.h"
 #include "Math/Axis.h"
-#include "Math/BasicMathExpressionEvaluator.h"
 #include "Math/Box.h"
 #include "Math/Box2D.h"
 #include "Math/BoxSphereBounds.h"
@@ -164,7 +163,6 @@
 #include "Misc/EngineVersionBase.h"
 #include "Misc/EnumClassFlags.h"
 #include "Misc/Exec.h"
-#include "Misc/ExpressionParserTypes.h"
 #include "Misc/FeedbackContext.h"
 #include "Misc/FileHelper.h"
 #include "Misc/FilterCollection.h"
@@ -258,3 +256,7 @@
 #if UE_ENABLE_INCLUDE_ORDER_DEPRECATED_IN_5_5
 #include "Templates/IsTriviallyDestructible.h"
 #endif
+
+#if UE_ENABLE_INCLUDE_ORDER_DEPRECATED_IN_5_6
+#include "Math/BasicMathExpressionEvaluator.h"
+#endif
diff --git a/Engine/Source/Runtime/Core/Public/Misc/FrameRate.h b/Engine/Source/Runtime/Core/Public/Misc/FrameRate.h
--- a/Engine/Source/Runtime/Core/Public/Misc/FrameRate.h
+++ b/Engine/Source/Runtime/Core/Public/Misc/FrameRate.h
@@ -8,7 +8,6 @@
 #include "Math/Range.h"
 #include "Math/RangeBound.h"
 #include "Math/UnrealMathUtility.h"
-#include "Misc/ExpressionParserTypes.h"
 #include "Misc/FrameNumber.h"
 #include "Misc/FrameTime.h"
 #include "Templates/ValueOrError.h"
@@ -322,6 +321,7 @@ inline bool FFrameRate::IsFactorOf(FFrameRate Other) const
 
 /**
  * Attempt to parse a frame rate from a string
+ * Note: user code may need to #include "Misc/ExpressionParserTypes.h"
  */
 CORE_API TValueOrError<FFrameRate, FExpressionError> ParseFrameRate(const TCHAR* FrameRateString);
 
@@ -329,3 +329,7 @@ CORE_API TValueOrError<FFrameRate, FExpressionError> ParseFrameRate(const TCHAR*
  * Common Lex::TryParseString overload for FFrameRate
  */
 CORE_API bool TryParseString(FFrameRate& OutFrameRate, const TCHAR* InString);
+
+#if UE_ENABLE_INCLUDE_ORDER_DEPRECATED_IN_5_6
+# include "ExpressionParserTypes.h"
+#endif
diff --git a/Engine/Source/Runtime/TimeManagement/Private/CommonFrameRates.cpp b/Engine/Source/Runtime/TimeManagement/Private/CommonFrameRates.cpp
--- a/Engine/Source/Runtime/TimeManagement/Private/CommonFrameRates.cpp
+++ b/Engine/Source/Runtime/TimeManagement/Private/CommonFrameRates.cpp
@@ -2,6 +2,7 @@
 
 #include "CommonFrameRates.h"
 #include "Algo/Find.h"
+#include "Internationalization/Internationalization.h"
 
 #define LOCTEXT_NAMESPACE "CommonFrameRates"
 
