Updated headers using UnrealCodeFixup to make sure dllstorage is on methods/staticvars instead of types. Part 5/n

diff --git a/Engine/Source/Editor/SceneOutliner/Public/ActorDescTreeItem.h b/Engine/Source/Editor/SceneOutliner/Public/ActorDescTreeItem.h
--- a/Engine/Source/Editor/SceneOutliner/Public/ActorDescTreeItem.h
+++ b/Engine/Source/Editor/SceneOutliner/Public/ActorDescTreeItem.h
@@ -8,13 +8,15 @@
 #include "UObject/ObjectKey.h"
 #include "WorldPartition/WorldPartitionHandle.h"
 
+#define UE_API SCENEOUTLINER_API
+
 class UActorDescContainerInstance;
 class UExternalDataLayerAsset;
 class FWorldPartitionActorDesc;
 class UToolMenu;
 
 /** A tree item that represents an actor in the world */
-struct SCENEOUTLINER_API FActorDescTreeItem : IActorBaseTreeItem
+struct FActorDescTreeItem : IActorBaseTreeItem
 {
 public:
 	DECLARE_DELEGATE_RetVal_OneParam(bool, FFilterPredicate, const FWorldPartitionActorDescInstance*);
@@ -36,48 +38,48 @@ public:
 	/** Constant identifier for this tree item */
 	FSceneOutlinerTreeItemID ID;
 
-	static FSceneOutlinerTreeItemID ComputeTreeItemID(FGuid InActorGuid, UActorDescContainerInstance* InContainer);
+	static UE_API FSceneOutlinerTreeItemID ComputeTreeItemID(FGuid InActorGuid, UActorDescContainerInstance* InContainer);
 
-	static bool ShouldDisplayInOutliner(const FWorldPartitionActorDescInstance* InActorDescInstance);
+	static UE_API bool ShouldDisplayInOutliner(const FWorldPartitionActorDescInstance* InActorDescInstance);
 
 	UE_DEPRECATED(5.4, "Use FWorldPartitionActorDescInstance version instead")
 	static bool ShouldDisplayInOutliner(const FWorldPartitionActorDesc* ActorDesc) { return false;}
 
 	/** Static type identifier for this tree item class */
-	static const FSceneOutlinerTreeItemType Type;
+	static UE_API const FSceneOutlinerTreeItemType Type;
 
 	/** Construct this item from an actor desc */
-	FActorDescTreeItem(const FGuid& InActorGuid, UActorDescContainerInstance* InContainerInstance);
+	UE_API FActorDescTreeItem(const FGuid& InActorGuid, UActorDescContainerInstance* InContainerInstance);
 
 	/** Construct this item from an actor desc instance */
-	FActorDescTreeItem(const FWorldPartitionActorDescInstance* InActorDescInstance);
+	UE_API FActorDescTreeItem(const FWorldPartitionActorDescInstance* InActorDescInstance);
 
 	UE_DEPRECATED(5.4, "Use FWorldPartitionActorDescInstance version instead")
 	FActorDescTreeItem(const FGuid& InActorGuid, class UActorDescContainer* InContainer) : IActorBaseTreeItem(Type) {}
 
 	/* Begin ISceneOutlinerTreeItem Implementation */
 	virtual bool IsValid() const override { return *ActorDescHandle != nullptr; }
-	virtual FSceneOutlinerTreeItemID GetID() const override;
-	virtual FString GetDisplayString() const override;
-	virtual bool CanInteract() const override;
-	virtual TSharedRef<SWidget> GenerateLabelWidget(ISceneOutliner& Outliner, const STableRow<FSceneOutlinerTreeItemPtr>& InRow) override;
-	virtual void GenerateContextMenu(UToolMenu* Menu, SSceneOutliner& Outliner) override;
+	UE_API virtual FSceneOutlinerTreeItemID GetID() const override;
+	UE_API virtual FString GetDisplayString() const override;
+	UE_API virtual bool CanInteract() const override;
+	UE_API virtual TSharedRef<SWidget> GenerateLabelWidget(ISceneOutliner& Outliner, const STableRow<FSceneOutlinerTreeItemPtr>& InRow) override;
+	UE_API virtual void GenerateContextMenu(UToolMenu* Menu, SSceneOutliner& Outliner) override;
 	virtual bool HasVisibilityInfo() const override { return true; }
-	virtual bool GetVisibility() const override;
-	virtual bool ShouldShowPinnedState() const override;
+	UE_API virtual bool GetVisibility() const override;
+	UE_API virtual bool ShouldShowPinnedState() const override;
 	virtual bool ShouldShowVisibilityState() const override { return false; }
 	virtual bool HasPinnedStateInfo() const override { return true; }
-	virtual bool GetPinnedState() const override;
-	virtual FString GetPackageName() const override;
+	UE_API virtual bool GetPinnedState() const override;
+	UE_API virtual FString GetPackageName() const override;
 	/* End ISceneOutlinerTreeItem Implementation */
 	
 	/* Begin IActorBaseTreeItem Implementation */
 	virtual const FGuid& GetGuid() const override { return ActorGuid; }
 	/* End IActorBaseTreeItem Implementation */
 
-	UExternalDataLayerAsset* GetExternalDataLayerAsset() const;
+	UE_API UExternalDataLayerAsset* GetExternalDataLayerAsset() const;
 
-	void FocusActorBounds() const;
+	UE_API void FocusActorBounds() const;
 
 protected:
 	FString DisplayString;
@@ -85,6 +87,8 @@ protected:
 private:
 	mutable TSoftObjectPtr<UExternalDataLayerAsset> CachedExternalDataLayerAsset;
 
-	void CopyActorFilePathtoClipboard() const;
+	UE_API void CopyActorFilePathtoClipboard() const;
 	FGuid ActorGuid;
 };
+
+#undef UE_API
diff --git a/Engine/Source/Editor/SceneOutliner/Public/ActorFolderPickingMode.h b/Engine/Source/Editor/SceneOutliner/Public/ActorFolderPickingMode.h
--- a/Engine/Source/Editor/SceneOutliner/Public/ActorFolderPickingMode.h
+++ b/Engine/Source/Editor/SceneOutliner/Public/ActorFolderPickingMode.h
@@ -4,24 +4,28 @@
 
 #include "ActorMode.h"
 
-class SCENEOUTLINER_API FActorFolderPickingMode : public FActorMode
+#define UE_API SCENEOUTLINER_API
+
+class FActorFolderPickingMode : public FActorMode
 {
 public:
-	FActorFolderPickingMode(SSceneOutliner* InSceneOutliner, FOnSceneOutlinerItemPicked InOnItemPicked, TWeakObjectPtr<UWorld> InSpecifiedWorldToDisplay = nullptr, const FFolder::FRootObject& InRootObject = FFolder::GetInvalidRootObject());
+	UE_API FActorFolderPickingMode(SSceneOutliner* InSceneOutliner, FOnSceneOutlinerItemPicked InOnItemPicked, TWeakObjectPtr<UWorld> InSpecifiedWorldToDisplay = nullptr, const FFolder::FRootObject& InRootObject = FFolder::GetInvalidRootObject());
 	virtual ~FActorFolderPickingMode() {}
 
 	/* Begin ISceneOutlinerMode Implementation */
-	virtual void OnFilterTextCommited(FSceneOutlinerItemSelection& Selection, ETextCommit::Type CommitType) override;
-	virtual void OnItemSelectionChanged(FSceneOutlinerTreeItemPtr Item, ESelectInfo::Type SelectionType, const FSceneOutlinerItemSelection& Selection) override;
-	virtual void CreateViewContent(FMenuBuilder& MenuBuilder) override;
+	UE_API virtual void OnFilterTextCommited(FSceneOutlinerItemSelection& Selection, ETextCommit::Type CommitType) override;
+	UE_API virtual void OnItemSelectionChanged(FSceneOutlinerTreeItemPtr Item, ESelectInfo::Type SelectionType, const FSceneOutlinerItemSelection& Selection) override;
+	UE_API virtual void CreateViewContent(FMenuBuilder& MenuBuilder) override;
 	virtual bool ShowViewButton() const override { return true; }
 	virtual bool ShouldShowFolders() const { return true; }
 protected:
-	virtual TUniquePtr<ISceneOutlinerHierarchy> CreateHierarchy() override;
+	UE_API virtual TUniquePtr<ISceneOutlinerHierarchy> CreateHierarchy() override;
 	/* End ISceneOutlinerMode Implementation */
 
 	/** Delegate to call when an item is picked */
 	FOnSceneOutlinerItemPicked OnItemPicked;
 
 	FFolder::FRootObject RootObject;
-};
\ No newline at end of file
+};
+
+#undef UE_API
diff --git a/Engine/Source/Editor/SceneOutliner/Public/ActorFolderTreeItem.h b/Engine/Source/Editor/SceneOutliner/Public/ActorFolderTreeItem.h
--- a/Engine/Source/Editor/SceneOutliner/Public/ActorFolderTreeItem.h
+++ b/Engine/Source/Editor/SceneOutliner/Public/ActorFolderTreeItem.h
@@ -5,40 +5,44 @@
 #include "FolderTreeItem.h"
 #include "ActorFolder.h"
 
-struct SCENEOUTLINER_API FActorFolderTreeItem : public FFolderTreeItem
+#define UE_API SCENEOUTLINER_API
+
+struct FActorFolderTreeItem : public FFolderTreeItem
 {
 public:
 	/** Static type identifier for this tree item class */
-	static const FSceneOutlinerTreeItemType Type;
+	static UE_API const FSceneOutlinerTreeItemType Type;
 
-	FActorFolderTreeItem(const FFolder& InFolder, const TWeakObjectPtr<UWorld>& InWorld);
+	UE_API FActorFolderTreeItem(const FFolder& InFolder, const TWeakObjectPtr<UWorld>& InWorld);
 
 	/** The world which this folder belongs to */
 	TWeakObjectPtr<UWorld> World;
 
 	/* Begin ISceneOutlinerTreeItem Implementation */
 	virtual bool IsValid() const override { return World.IsValid(); }
-	virtual void OnExpansionChanged() override;
-	virtual void Delete(const FFolder& InNewParentFolder) override;
-	virtual bool CanInteract() const override;
-	virtual TSharedRef<SWidget> GenerateLabelWidget(ISceneOutliner& Outliner, const STableRow<FSceneOutlinerTreeItemPtr>& InRow) override;
+	UE_API virtual void OnExpansionChanged() override;
+	UE_API virtual void Delete(const FFolder& InNewParentFolder) override;
+	UE_API virtual bool CanInteract() const override;
+	UE_API virtual TSharedRef<SWidget> GenerateLabelWidget(ISceneOutliner& Outliner, const STableRow<FSceneOutlinerTreeItemPtr>& InRow) override;
 	virtual bool ShouldShowPinnedState() const override { return false; }
-	virtual bool ShouldShowVisibilityState() const override;
-	virtual bool ShouldRemoveOnceLastChildRemoved() const override;
-	virtual FFolder GetFolder() const override;
-	virtual FString GetPackageName() const override;
+	UE_API virtual bool ShouldShowVisibilityState() const override;
+	UE_API virtual bool ShouldRemoveOnceLastChildRemoved() const override;
+	UE_API virtual FFolder GetFolder() const override;
+	UE_API virtual FString GetPackageName() const override;
 	/* End FFolderTreeItem Implementation */
 		
 	/* Begin FFolderTreeItem Implementation */
-	virtual void MoveTo(const FFolder& InNewParentFolder) override;
-	virtual void SetPath(const FName& InNewPath) override;
+	UE_API virtual void MoveTo(const FFolder& InNewParentFolder) override;
+	UE_API virtual void SetPath(const FName& InNewPath) override;
 	
-	bool CanChangeChildrenPinnedState() const;
+	UE_API bool CanChangeChildrenPinnedState() const;
 	const UActorFolder* GetActorFolder() const { return ActorFolder.Get(); }
 private:
-	virtual void CreateSubFolder(TWeakPtr<SSceneOutliner> WeakOutliner) override;
+	UE_API virtual void CreateSubFolder(TWeakPtr<SSceneOutliner> WeakOutliner) override;
 	/* End FFolderTreeItem Implementation */
 
 	/** The actor folder object (can be invalid) */
 	TWeakObjectPtr<UActorFolder> ActorFolder;
 };
+
+#undef UE_API
diff --git a/Engine/Source/Editor/SceneOutliner/Public/ActorHierarchy.h b/Engine/Source/Editor/SceneOutliner/Public/ActorHierarchy.h
--- a/Engine/Source/Editor/SceneOutliner/Public/ActorHierarchy.h
+++ b/Engine/Source/Editor/SceneOutliner/Public/ActorHierarchy.h
@@ -59,10 +59,10 @@ protected:
 	SCENEOUTLINER_API virtual void OnActorDescInstanceRemoved(FWorldPartitionActorDescInstance* InActorDescInstance);
 
 	UE_DEPRECATED(5.4, "Use OnActorDescInstanceAdded instead")
-	SCENEOUTLINER_API virtual void OnActorDescAdded(FWorldPartitionActorDesc* ActorDesc) {}
+	virtual void OnActorDescAdded(FWorldPartitionActorDesc* ActorDesc) {}
 
 	UE_DEPRECATED(5.4, "Use OnActorDescInstanceRemoved instead")
-	SCENEOUTLINER_API virtual void OnActorDescRemoved(FWorldPartitionActorDesc* ActorDesc) {}
+	virtual void OnActorDescRemoved(FWorldPartitionActorDesc* ActorDesc) {}
 	
 	void OnComponentsUpdated();
 
diff --git a/Engine/Source/Editor/SceneOutliner/Public/ActorMode.h b/Engine/Source/Editor/SceneOutliner/Public/ActorMode.h
--- a/Engine/Source/Editor/SceneOutliner/Public/ActorMode.h
+++ b/Engine/Source/Editor/SceneOutliner/Public/ActorMode.h
@@ -5,18 +5,20 @@
 #include "ISceneOutlinerMode.h"
 #include "WorldPartition/WorldPartitionHandle.h"
 
+#define UE_API SCENEOUTLINER_API
+
 namespace SceneOutliner
 {
 	/** Functor which can be used to get weak actor pointers from a selection */
-	struct SCENEOUTLINER_API FWeakActorSelector
+	struct FWeakActorSelector
 	{
-		bool operator()(const TWeakPtr<ISceneOutlinerTreeItem>& Item, TWeakObjectPtr<AActor>& DataOut) const;
+		UE_API bool operator()(const TWeakPtr<ISceneOutlinerTreeItem>& Item, TWeakObjectPtr<AActor>& DataOut) const;
 	};
 
 	/** Functor which can be used to get actors from a selection including component parents */
-	struct SCENEOUTLINER_API FActorSelector
+	struct FActorSelector
 	{
-		bool operator()(const TWeakPtr<ISceneOutlinerTreeItem>& Item, AActor*& ActorPtrOut) const;
+		UE_API bool operator()(const TWeakPtr<ISceneOutlinerTreeItem>& Item, AActor*& ActorPtrOut) const;
 	};
 
 	struct UE_DEPRECATED(5.4, "Use FActorHandleSelector instead") FActorDescSelector
@@ -25,13 +27,13 @@ namespace SceneOutliner
 	};
 
 	/** Functor which can be used to get actor descriptors from a selection  */
-	struct SCENEOUTLINER_API FActorHandleSelector
+	struct FActorHandleSelector
 	{
-		bool operator()(const TWeakPtr<ISceneOutlinerTreeItem>& Item, FWorldPartitionHandle& ActorHandleOut) const;
+		UE_API bool operator()(const TWeakPtr<ISceneOutlinerTreeItem>& Item, FWorldPartitionHandle& ActorHandleOut) const;
 	};
 }
 
-struct SCENEOUTLINER_API FActorModeParams
+struct FActorModeParams
 {
 	FActorModeParams() {}
 
@@ -60,7 +62,7 @@ struct SCENEOUTLINER_API FActorModeParams
 	bool bSearchComponentsByActorName = false;
 };
 
-class SCENEOUTLINER_API FActorMode : public ISceneOutlinerMode
+class FActorMode : public ISceneOutlinerMode
 {
 public:
 	struct EItemSortOrder
@@ -68,57 +70,57 @@ public:
 		enum Type { World = 0, Level = 10, Folder = 20, Actor = 30 };
 	};
 	
-	FActorMode(const FActorModeParams& Params);
-	virtual ~FActorMode();
+	UE_API FActorMode(const FActorModeParams& Params);
+	UE_API virtual ~FActorMode();
 
-	virtual void Rebuild() override;
+	UE_API virtual void Rebuild() override;
 
-	void BuildWorldPickerMenu(FMenuBuilder& MenuBuilder);
+	UE_API void BuildWorldPickerMenu(FMenuBuilder& MenuBuilder);
 
 	virtual void SynchronizeSelection() override { SynchronizeActorSelection(); }
 
-	virtual void OnFilterTextChanged(const FText& InFilterText) override;
+	UE_API virtual void OnFilterTextChanged(const FText& InFilterText) override;
 
-	virtual int32 GetTypeSortPriority(const ISceneOutlinerTreeItem& Item) const override;
+	UE_API virtual int32 GetTypeSortPriority(const ISceneOutlinerTreeItem& Item) const override;
 
-	static bool IsActorDisplayable(const SSceneOutliner* SceneOutliner, const AActor* Actor, bool bShowLevelInstanceContent = false);
-	static bool IsActorLevelDisplayable(ULevel* InLevel);
+	static UE_API bool IsActorDisplayable(const SSceneOutliner* SceneOutliner, const AActor* Actor, bool bShowLevelInstanceContent = false);
+	static UE_API bool IsActorLevelDisplayable(ULevel* InLevel);
 
-	virtual FFolder::FRootObject GetRootObject() const override;
-	virtual FFolder::FRootObject GetPasteTargetRootObject() const override;
+	UE_API virtual FFolder::FRootObject GetRootObject() const override;
+	UE_API virtual FFolder::FRootObject GetPasteTargetRootObject() const override;
 
-	virtual bool CanInteract(const ISceneOutlinerTreeItem& Item) const override;
-	virtual bool CanPopulate() const override;
+	UE_API virtual bool CanInteract(const ISceneOutlinerTreeItem& Item) const override;
+	UE_API virtual bool CanPopulate() const override;
 	
-	virtual TSharedPtr<FDragDropOperation> CreateDragDropOperation(const FPointerEvent& MouseEvent, const TArray<FSceneOutlinerTreeItemPtr>& InTreeItems) const override;
-	virtual bool ParseDragDrop(FSceneOutlinerDragDropPayload& OutPayload, const FDragDropOperation& Operation) const override;
-	virtual FSceneOutlinerDragValidationInfo ValidateDrop(const ISceneOutlinerTreeItem& DropTarget, const FSceneOutlinerDragDropPayload& Payload) const override;
-	virtual void OnDrop(ISceneOutlinerTreeItem& DropTarget, const FSceneOutlinerDragDropPayload& Payload, const FSceneOutlinerDragValidationInfo& ValidationInfo) const override;
+	UE_API virtual TSharedPtr<FDragDropOperation> CreateDragDropOperation(const FPointerEvent& MouseEvent, const TArray<FSceneOutlinerTreeItemPtr>& InTreeItems) const override;
+	UE_API virtual bool ParseDragDrop(FSceneOutlinerDragDropPayload& OutPayload, const FDragDropOperation& Operation) const override;
+	UE_API virtual FSceneOutlinerDragValidationInfo ValidateDrop(const ISceneOutlinerTreeItem& DropTarget, const FSceneOutlinerDragDropPayload& Payload) const override;
+	UE_API virtual void OnDrop(ISceneOutlinerTreeItem& DropTarget, const FSceneOutlinerDragDropPayload& Payload, const FSceneOutlinerDragValidationInfo& ValidationInfo) const override;
 
 
 private:
 	/** Called when the user selects a world in the world picker menu */
-	void OnSelectWorld(TWeakObjectPtr<UWorld> World);
+	UE_API void OnSelectWorld(TWeakObjectPtr<UWorld> World);
 private:
 	/* Private Helpers */
 
-	void ChooseRepresentingWorld();
-	bool IsWorldChecked(TWeakObjectPtr<UWorld> World) const;
-	bool GetFolderNamesFromPayload(const FSceneOutlinerDragDropPayload& InPayload, TArray<FName>& OutFolders, FFolder::FRootObject& OutCommonRootObject) const;
+	UE_API void ChooseRepresentingWorld();
+	UE_API bool IsWorldChecked(TWeakObjectPtr<UWorld> World) const;
+	UE_API bool GetFolderNamesFromPayload(const FSceneOutlinerDragDropPayload& InPayload, TArray<FName>& OutFolders, FFolder::FRootObject& OutCommonRootObject) const;
 
 protected:
-	void SynchronizeActorSelection();
-	virtual bool IsActorDisplayable(const AActor* InActor) const;
+	UE_API void SynchronizeActorSelection();
+	UE_API virtual bool IsActorDisplayable(const AActor* InActor) const;
 
 	/** Set the Scene Outliner attached to this mode as the most recently used outliner in the Level Editor */
-	void SetAsMostRecentOutliner() const;
+	UE_API void SetAsMostRecentOutliner() const;
 
-	virtual TUniquePtr<ISceneOutlinerHierarchy> CreateHierarchy() override;
+	UE_API virtual TUniquePtr<ISceneOutlinerHierarchy> CreateHierarchy() override;
 
 	// Called when actors are attached to a parent actor via drag and drop
 	virtual void OnActorsAttached(AActor* ParentActor, TArray<TWeakObjectPtr<AActor>> ChildActors) const {}
 	
-	FFolder GetWorldDefaultRootFolder() const;
+	UE_API FFolder GetWorldDefaultRootFolder() const;
 protected:
 	/** The world which we are currently representing */
 	TWeakObjectPtr<UWorld> RepresentingWorld;
@@ -151,3 +153,5 @@ protected:
 	/** Should Outliner Tree get collapsed on selection, except for the item that was just selected */
 	bool bCollapseOutlinerTreeOnNewSelection;
 };
+
+#undef UE_API
diff --git a/Engine/Source/Editor/SceneOutliner/Public/ActorModeInteractive.h b/Engine/Source/Editor/SceneOutliner/Public/ActorModeInteractive.h
--- a/Engine/Source/Editor/SceneOutliner/Public/ActorModeInteractive.h
+++ b/Engine/Source/Editor/SceneOutliner/Public/ActorModeInteractive.h
@@ -10,7 +10,7 @@ public:
 	SCENEOUTLINER_API FActorModeInteractive(const FActorModeParams& Params);
 	SCENEOUTLINER_API virtual ~FActorModeInteractive();
 
-	SCENEOUTLINER_API virtual bool IsInteractive() const override { return true; }
+	virtual bool IsInteractive() const override { return true; }
 protected:
 	/* Events */
 
diff --git a/Engine/Source/Editor/SceneOutliner/Public/ActorPickingMode.h b/Engine/Source/Editor/SceneOutliner/Public/ActorPickingMode.h
--- a/Engine/Source/Editor/SceneOutliner/Public/ActorPickingMode.h
+++ b/Engine/Source/Editor/SceneOutliner/Public/ActorPickingMode.h
@@ -4,19 +4,21 @@
 
 #include "ActorMode.h"
 
-class SCENEOUTLINER_API FActorPickingMode : public FActorMode
+#define UE_API SCENEOUTLINER_API
+
+class FActorPickingMode : public FActorMode
 {
 public:
-	FActorPickingMode(const FActorModeParams& Params, FOnSceneOutlinerItemPicked OnItemPickedDelegate);
+	UE_API FActorPickingMode(const FActorModeParams& Params, FOnSceneOutlinerItemPicked OnItemPickedDelegate);
 
 	virtual ~FActorPickingMode() {};
 public:
-	virtual void OnItemSelectionChanged(FSceneOutlinerTreeItemPtr Item, ESelectInfo::Type SelectionType, const FSceneOutlinerItemSelection& Selection) override;
+	UE_API virtual void OnItemSelectionChanged(FSceneOutlinerTreeItemPtr Item, ESelectInfo::Type SelectionType, const FSceneOutlinerItemSelection& Selection) override;
 
 	/** Allow the user to commit their selection by pressing enter if it is valid */
-	virtual void OnFilterTextCommited(FSceneOutlinerItemSelection& Selection, ETextCommit::Type CommitType) override;
+	UE_API virtual void OnFilterTextCommited(FSceneOutlinerItemSelection& Selection, ETextCommit::Type CommitType) override;
 
-	virtual void CreateViewContent(FMenuBuilder& MenuBuilder) override;
+	UE_API virtual void CreateViewContent(FMenuBuilder& MenuBuilder) override;
 
 	// Don't synchronize world selection
 	virtual void SynchronizeSelection() override {};
@@ -24,4 +26,6 @@ public:
 	virtual bool ShowViewButton() const override { return true; }
 private:
 	FOnSceneOutlinerItemPicked OnItemPicked;
-};
\ No newline at end of file
+};
+
+#undef UE_API
diff --git a/Engine/Source/Editor/SceneOutliner/Public/ActorTreeItem.h b/Engine/Source/Editor/SceneOutliner/Public/ActorTreeItem.h
--- a/Engine/Source/Editor/SceneOutliner/Public/ActorTreeItem.h
+++ b/Engine/Source/Editor/SceneOutliner/Public/ActorTreeItem.h
@@ -7,8 +7,10 @@
 #include "ActorBaseTreeItem.h"
 #include "UObject/ObjectKey.h"
 
+#define UE_API SCENEOUTLINER_API
+
 /** A tree item that represents an actor in the world */
-struct SCENEOUTLINER_API FActorTreeItem : IActorBaseTreeItem
+struct FActorTreeItem : IActorBaseTreeItem
 {
 public:
 	DECLARE_DELEGATE_RetVal_OneParam(bool, FFilterPredicate, const AActor*);
@@ -31,32 +33,32 @@ public:
 	const FObjectKey ID;
 
 	/** Static type identifier for this tree item class */
-	static const FSceneOutlinerTreeItemType Type;
+	static UE_API const FSceneOutlinerTreeItemType Type;
 
 	/** Construct this item from an actor */
-	FActorTreeItem(AActor* InActor);
+	UE_API FActorTreeItem(AActor* InActor);
 
 	/* Begin ISceneOutlinerTreeItem Implementation */
 	virtual bool IsValid() const override { return Actor.IsValid(); }
-	virtual FSceneOutlinerTreeItemID GetID() const override;
-	virtual FFolder::FRootObject GetRootObject() const override;
-	virtual FString GetDisplayString() const override;
-	virtual bool CanInteract() const override;
-	virtual TSharedRef<SWidget> GenerateLabelWidget(ISceneOutliner& Outliner, const STableRow<FSceneOutlinerTreeItemPtr>& InRow) override;
-	virtual void OnVisibilityChanged(const bool bNewVisibility) override;
+	UE_API virtual FSceneOutlinerTreeItemID GetID() const override;
+	UE_API virtual FFolder::FRootObject GetRootObject() const override;
+	UE_API virtual FString GetDisplayString() const override;
+	UE_API virtual bool CanInteract() const override;
+	UE_API virtual TSharedRef<SWidget> GenerateLabelWidget(ISceneOutliner& Outliner, const STableRow<FSceneOutlinerTreeItemPtr>& InRow) override;
+	UE_API virtual void OnVisibilityChanged(const bool bNewVisibility) override;
 	virtual bool HasVisibilityInfo() const override { return true; }
-	virtual bool GetVisibility() const override;
-	virtual bool ShouldShowPinnedState() const override;
-	virtual bool ShouldShowVisibilityState() const override;
+	UE_API virtual bool GetVisibility() const override;
+	UE_API virtual bool ShouldShowPinnedState() const override;
+	UE_API virtual bool ShouldShowVisibilityState() const override;
 	virtual bool HasPinnedStateInfo() const override { return true; };
-	virtual bool GetPinnedState() const override;
-	virtual void OnLabelChanged() override;
-	virtual void GenerateContextMenu(UToolMenu* Menu, SSceneOutliner& Outliner) override;
-	virtual FString GetPackageName() const override;
+	UE_API virtual bool GetPinnedState() const override;
+	UE_API virtual void OnLabelChanged() override;
+	UE_API virtual void GenerateContextMenu(UToolMenu* Menu, SSceneOutliner& Outliner) override;
+	UE_API virtual FString GetPackageName() const override;
 	/* End ISceneOutlinerTreeItem Implementation */
 
 	/* Begin IActorBaseTreeItem Implementation */
-	virtual const FGuid& GetGuid() const override;
+	UE_API virtual const FGuid& GetGuid() const override;
 	/* End IActorBaseTreeItem Implementation */
 
 	/** true if this item exists in both the current world and PIE. */
@@ -64,13 +66,15 @@ public:
 
 protected:
 	/** For use by derived classes, as it allows for passing down their own type. */
-	FActorTreeItem(FSceneOutlinerTreeItemType TypeIn, AActor* InActor);
+	UE_API FActorTreeItem(FSceneOutlinerTreeItemType TypeIn, AActor* InActor);
 
-	virtual void UpdateDisplayString();
+	UE_API virtual void UpdateDisplayString();
 
 	/** Cached display string */
 	FString DisplayString;
 
 private:
-	void UpdateDisplayStringInternal();
+	UE_API void UpdateDisplayStringInternal();
 };
+
+#undef UE_API
diff --git a/Engine/Source/Editor/SceneOutliner/Public/ComponentTreeItem.h b/Engine/Source/Editor/SceneOutliner/Public/ComponentTreeItem.h
--- a/Engine/Source/Editor/SceneOutliner/Public/ComponentTreeItem.h
+++ b/Engine/Source/Editor/SceneOutliner/Public/ComponentTreeItem.h
@@ -8,8 +8,10 @@
 #include "UObject/ObjectKey.h"
 #include "Components/SceneComponent.h"
 
+#define UE_API SCENEOUTLINER_API
+
 /** A tree item that represents an Component in the world */
-struct SCENEOUTLINER_API FComponentTreeItem : ISceneOutlinerTreeItem
+struct FComponentTreeItem : ISceneOutlinerTreeItem
 {
 public:
 	DECLARE_DELEGATE_RetVal_OneParam(bool, FFilterPredicate, const UActorComponent*);
@@ -32,19 +34,19 @@ public:
 	const FObjectKey ID;
 
 	/** Static type identifier for this tree item class */
-	static const FSceneOutlinerTreeItemType Type;
+	static UE_API const FSceneOutlinerTreeItemType Type;
 
 	/** Construct this item from an Component */
-	FComponentTreeItem(UActorComponent* InComponent, bool bInSearchComponentsByActorName = false);
+	UE_API FComponentTreeItem(UActorComponent* InComponent, bool bInSearchComponentsByActorName = false);
 
-	bool GetSearchComponentByActorName() const;
+	UE_API bool GetSearchComponentByActorName() const;
 
 	/* Begin ISceneOutlinerTreeItem Implementation */
 	virtual bool IsValid() const override { return Component.IsValid(); }
-	virtual FSceneOutlinerTreeItemID GetID() const override;
-	virtual FString GetDisplayString() const override;
-	virtual bool CanInteract() const override;
-	virtual TSharedRef<SWidget> GenerateLabelWidget(ISceneOutliner& Outliner, const STableRow<FSceneOutlinerTreeItemPtr>& InRow) override;
+	UE_API virtual FSceneOutlinerTreeItemID GetID() const override;
+	UE_API virtual FString GetDisplayString() const override;
+	UE_API virtual bool CanInteract() const override;
+	UE_API virtual TSharedRef<SWidget> GenerateLabelWidget(ISceneOutliner& Outliner, const STableRow<FSceneOutlinerTreeItemPtr>& InRow) override;
 	virtual bool ShouldShowVisibilityState() const { return false; }
 	virtual bool HasVisibilityInfo() const { return false; }
 	virtual bool GetVisibility() const { return false; }
@@ -59,3 +61,5 @@ public:
 	/** Cache the string displayed */
 	FString CachedDisplayString;
 };
+
+#undef UE_API
diff --git a/Engine/Source/Editor/SceneOutliner/Public/FolderTreeItem.h b/Engine/Source/Editor/SceneOutliner/Public/FolderTreeItem.h
--- a/Engine/Source/Editor/SceneOutliner/Public/FolderTreeItem.h
+++ b/Engine/Source/Editor/SceneOutliner/Public/FolderTreeItem.h
@@ -6,6 +6,8 @@
 #include "Folder.h"
 #include "ISceneOutlinerTreeItem.h"
 
+#define UE_API SCENEOUTLINER_API
+
 class UToolMenu;
 
 namespace SceneOutliner
@@ -18,11 +20,11 @@ namespace SceneOutliner
 
 
 /** A tree item that represents a folder in the world */
-struct SCENEOUTLINER_API FFolderTreeItem : ISceneOutlinerTreeItem
+struct FFolderTreeItem : ISceneOutlinerTreeItem
 {
 public:
 	/** Static type identifier for this tree item class */
-	static const FSceneOutlinerTreeItemType Type;
+	static UE_API const FSceneOutlinerTreeItemType Type;
 
 	DECLARE_DELEGATE_RetVal_OneParam(bool, FFilterPredicate, const FFolder&);
 	DECLARE_DELEGATE_RetVal_OneParam(bool, FInteractivePredicate, const FFolder&);
@@ -43,24 +45,24 @@ public:
 	/** The leaf name of this folder */
 	FName LeafName;
 
-	FFolderTreeItem(const FFolder& InFolder, FSceneOutlinerTreeItemType InType);
-	FFolderTreeItem(const FFolder& InFolder);
+	UE_API FFolderTreeItem(const FFolder& InFolder, FSceneOutlinerTreeItemType InType);
+	UE_API FFolderTreeItem(const FFolder& InFolder);
 	/** Constructor that takes a path to this folder (including leaf-name) */
-	FFolderTreeItem(FName InPath);
+	UE_API FFolderTreeItem(FName InPath);
 	/** Constructor that takes a path to this folder and a subclass tree item type (used for subclassing FFolderTreeItem) */
-	FFolderTreeItem(FName InPath, FSceneOutlinerTreeItemType Type);
+	UE_API FFolderTreeItem(FName InPath, FSceneOutlinerTreeItemType Type);
 
 	/* Begin ISceneOutlinerTreeItem Implementation */
 	virtual bool IsValid() const override { return true; }
-	virtual FSceneOutlinerTreeItemID GetID() const override;
-	virtual FFolder::FRootObject GetRootObject() const override;
-	virtual FString GetDisplayString() const override;
-	virtual bool CanInteract() const override;
-	virtual void GenerateContextMenu(UToolMenu* Menu, SSceneOutliner& Outliner) override;
+	UE_API virtual FSceneOutlinerTreeItemID GetID() const override;
+	UE_API virtual FFolder::FRootObject GetRootObject() const override;
+	UE_API virtual FString GetDisplayString() const override;
+	UE_API virtual bool CanInteract() const override;
+	UE_API virtual void GenerateContextMenu(UToolMenu* Menu, SSceneOutliner& Outliner) override;
 	/** Delete this folder, children will be reparented to provided new parent path */
 	virtual void Delete(const FFolder& InNewParentFolder) {}
 	
-	virtual void SetPath(const FName& InNewPath);
+	UE_API virtual void SetPath(const FName& InNewPath);
 	const FName& GetPath() const { return Path; }
 	const FName& GetLeafName() const { return LeafName; }
 	virtual FFolder GetFolder() const { return FFolder(RootObject, Path); }
@@ -75,8 +77,10 @@ private:
 	/** Create a new folder as a child of this one */
 	virtual void CreateSubFolder(TWeakPtr<SSceneOutliner> WeakOutliner) {}
 	/** Duplicate folder hierarchy */
-	void DuplicateHierarchy(TWeakPtr<SSceneOutliner> WeakOutliner);
+	UE_API void DuplicateHierarchy(TWeakPtr<SSceneOutliner> WeakOutliner);
 
 	/** Folder's root object (can be null) */
 	FFolder::FRootObject RootObject;
 };
+
+#undef UE_API
diff --git a/Engine/Source/Editor/SceneOutliner/Public/ISceneOutlinerTreeItem.h b/Engine/Source/Editor/SceneOutliner/Public/ISceneOutlinerTreeItem.h
--- a/Engine/Source/Editor/SceneOutliner/Public/ISceneOutlinerTreeItem.h
+++ b/Engine/Source/Editor/SceneOutliner/Public/ISceneOutlinerTreeItem.h
@@ -7,11 +7,13 @@
 #include "Widgets/Views/STableRow.h"
 #include "SceneOutlinerStandaloneTypes.h"
 
+#define UE_API SCENEOUTLINER_API
+
 class ISceneOutliner;
 class UToolMenu;
 
 /** Base tree item interface  */
-struct SCENEOUTLINER_API ISceneOutlinerTreeItem : TSharedFromThis<ISceneOutlinerTreeItem>
+struct ISceneOutlinerTreeItem : TSharedFromThis<ISceneOutlinerTreeItem>
 {
 	/** Friendship required for access to various internals */
 	friend SSceneOutliner;
@@ -62,7 +64,7 @@ protected:
 	mutable TSet<TWeakPtr<ISceneOutlinerTreeItem>> Children;
 
 	/** Static type identifier for the base class tree item */
-	static const FSceneOutlinerTreeItemType Type;
+	static UE_API const FSceneOutlinerTreeItemType Type;
 
 	/** Tree item type identifier */
 	FSceneOutlinerTreeItemType TreeType;
@@ -185,3 +187,5 @@ public:
 	/** Returns (optionally) the name of the UPackage this item represents/is-part-of */
 	virtual FString GetPackageName() const { return FString(); }
 };
+
+#undef UE_API
diff --git a/Engine/Source/Editor/SceneOutliner/Public/LevelTreeItem.h b/Engine/Source/Editor/SceneOutliner/Public/LevelTreeItem.h
--- a/Engine/Source/Editor/SceneOutliner/Public/LevelTreeItem.h
+++ b/Engine/Source/Editor/SceneOutliner/Public/LevelTreeItem.h
@@ -4,14 +4,16 @@
 
 #include "ISceneOutlinerTreeItem.h"
 
+#define UE_API SCENEOUTLINER_API
+
 class ULevel;
 
 /** A tree item that represents a level in the world */
-struct SCENEOUTLINER_API FLevelTreeItem : ISceneOutlinerTreeItem
+struct FLevelTreeItem : ISceneOutlinerTreeItem
 {
 public:
 	/** Static type identifier for this tree item class */
-	static const FSceneOutlinerTreeItemType Type;
+	static UE_API const FSceneOutlinerTreeItemType Type;
 
 	DECLARE_DELEGATE_RetVal_OneParam(bool, FFilterPredicate, const ULevel*);
 	DECLARE_DELEGATE_RetVal_OneParam(bool, FInteractivePredicate, const ULevel*);
@@ -33,14 +35,16 @@ public:
 	const FObjectKey ID;
 
 	/** Construct this item from a level */
-	FLevelTreeItem(ULevel* InLevel);
+	UE_API FLevelTreeItem(ULevel* InLevel);
 
 	/* Begin ISceneOutlinerTreeItem Implementation */
 	virtual bool IsValid() const override { return Level.IsValid(); }
 	virtual FSceneOutlinerTreeItemID GetID() const override { return ID; }
-	virtual FString GetDisplayString() const override;
-	virtual bool CanInteract() const override;
-	virtual TSharedRef<SWidget> GenerateLabelWidget(ISceneOutliner& Outliner, const STableRow<FSceneOutlinerTreeItemPtr>& InRow) override;
-	virtual void GenerateContextMenu(UToolMenu* Menu, SSceneOutliner& Outliner) override;
+	UE_API virtual FString GetDisplayString() const override;
+	UE_API virtual bool CanInteract() const override;
+	UE_API virtual TSharedRef<SWidget> GenerateLabelWidget(ISceneOutliner& Outliner, const STableRow<FSceneOutlinerTreeItemPtr>& InRow) override;
+	UE_API virtual void GenerateContextMenu(UToolMenu* Menu, SSceneOutliner& Outliner) override;
 	/* End ISceneOutlinerTreeItem Implementation */
 };
+
+#undef UE_API
diff --git a/Engine/Source/Editor/SceneOutliner/Public/SceneOutlinerDelegates.h b/Engine/Source/Editor/SceneOutliner/Public/SceneOutlinerDelegates.h
--- a/Engine/Source/Editor/SceneOutliner/Public/SceneOutlinerDelegates.h
+++ b/Engine/Source/Editor/SceneOutliner/Public/SceneOutlinerDelegates.h
@@ -8,7 +8,7 @@ class FSceneOutlinerDelegates
 {
 public:
 	/** Return a single FSceneOutlinerDelegates object */
-	SCENEOUTLINER_API static FSceneOutlinerDelegates& Get()
+	static FSceneOutlinerDelegates& Get()
 	{
 		// return the singleton object
 		static FSceneOutlinerDelegates Singleton;
diff --git a/Engine/Source/Editor/SceneOutliner/Public/SceneOutlinerDragDrop.h b/Engine/Source/Editor/SceneOutliner/Public/SceneOutlinerDragDrop.h
--- a/Engine/Source/Editor/SceneOutliner/Public/SceneOutlinerDragDrop.h
+++ b/Engine/Source/Editor/SceneOutliner/Public/SceneOutlinerDragDrop.h
@@ -16,6 +16,8 @@
 #include "Templates/SharedPointer.h"
 #include "Templates/TypeHash.h"
 
+#define UE_API SCENEOUTLINER_API
+
 class SWidget;
 struct FSlateBrush;
 struct ISceneOutlinerTreeItem;
@@ -169,11 +171,11 @@ struct FSceneOutlinerDragValidationInfo
 };
 
 /** A drag/drop operation that was started from the scene outliner */
-struct SCENEOUTLINER_API FSceneOutlinerDragDropOp : public FCompositeDragDropOp
+struct FSceneOutlinerDragDropOp : public FCompositeDragDropOp
 {
 	DRAG_DROP_OPERATOR_TYPE(FSceneOutlinerDragDropOp, FCompositeDragDropOp);
 		
-	FSceneOutlinerDragDropOp();
+	UE_API FSceneOutlinerDragDropOp();
 
 	using FDragDropOperation::Construct;
 
@@ -189,12 +191,12 @@ struct SCENEOUTLINER_API FSceneOutlinerDragDropOp : public FCompositeDragDropOp
 		OverrideIcon = InOverrideIcon;
 	}
 
-	virtual TSharedPtr<SWidget> GetDefaultDecorator() const override;
+	UE_API virtual TSharedPtr<SWidget> GetDefaultDecorator() const override;
 
 private:
 
-	EVisibility GetOverrideVisibility() const;
-	EVisibility GetDefaultVisibility() const;
+	UE_API EVisibility GetOverrideVisibility() const;
+	UE_API EVisibility GetDefaultVisibility() const;
 
 	FText OverrideText;
 	FText GetOverrideText() const { return OverrideText; }
@@ -202,3 +204,5 @@ private:
 	const FSlateBrush* OverrideIcon;
 	const FSlateBrush* GetOverrideIcon() const { return OverrideIcon; }
 };
+
+#undef UE_API
diff --git a/Engine/Source/Editor/SceneOutliner/Public/SceneOutlinerFilters.h b/Engine/Source/Editor/SceneOutliner/Public/SceneOutlinerFilters.h
--- a/Engine/Source/Editor/SceneOutliner/Public/SceneOutlinerFilters.h
+++ b/Engine/Source/Editor/SceneOutliner/Public/SceneOutlinerFilters.h
@@ -13,6 +13,8 @@
 #include "SceneOutlinerFwd.h"
 #include "Templates/SharedPointer.h"
 
+#define UE_API SCENEOUTLINER_API
+
 class FMenuBuilder;
 class FSceneOutlinerFilter;
 struct FSceneOutlinerFilters;
@@ -20,7 +22,7 @@ struct FSceneOutlinerFilters;
 /**
 	* Contains information used to create a filter which will be displayed as user toggleable filter
 	*/
-class SCENEOUTLINER_API FSceneOutlinerFilterInfo
+class FSceneOutlinerFilterInfo
 { 
 public:
 	FSceneOutlinerFilterInfo(const FText& InFilterTitle, const FText& InFilterTooltip, bool bInActive, const FCreateSceneOutlinerFilter& InFactory = FCreateSceneOutlinerFilter())
@@ -31,19 +33,19 @@ public:
 	{}
 
 	/** Initialize and apply a new filter */
-	void InitFilter(TSharedPtr<FSceneOutlinerFilters> InFilters);
+	UE_API void InitFilter(TSharedPtr<FSceneOutlinerFilters> InFilters);
 
 	/** Add menu for this filter */
-	void AddMenu(FMenuBuilder& InMenuBuilder);
+	UE_API void AddMenu(FMenuBuilder& InMenuBuilder);
 
-	bool IsFilterActive() const;
+	UE_API bool IsFilterActive() const;
 
 	DECLARE_EVENT_OneParam(FSceneOutlinerFilterInfo, FOnToggle, bool);
 	FOnToggle& OnToggle() { return OnToggleEvent; }
 
 private:
-	void ApplyFilter(bool bActive);
-	void ToggleFilterActive();
+	UE_API void ApplyFilter(bool bActive);
+	UE_API void ToggleFilterActive();
 
 	TWeakPtr<FSceneOutlinerFilters> Filters;
 
@@ -206,3 +208,5 @@ struct FSceneOutlinerFilters : public TFilterCollection<const ISceneOutlinerTree
 		Add(MakeShareable(new TSceneOutlinerPredicateFilter<T>(InFilterPred, InDefaultBehaviour, InInteractivePred)));
 	}
 };
+
+#undef UE_API
diff --git a/Engine/Source/Editor/SceneOutliner/Public/SceneOutlinerGutter.h b/Engine/Source/Editor/SceneOutliner/Public/SceneOutlinerGutter.h
--- a/Engine/Source/Editor/SceneOutliner/Public/SceneOutlinerGutter.h
+++ b/Engine/Source/Editor/SceneOutliner/Public/SceneOutlinerGutter.h
@@ -11,18 +11,20 @@
 #include "ISceneOutlinerColumn.h"
 #include "ISceneOutlinerTreeItem.h"
 
+#define UE_API SCENEOUTLINER_API
+
 class ISceneOutliner;
 template<typename ItemType> class STableRow;
 
 /**
  * A gutter for the SceneOutliner which handles setting and visualizing item visibility
  */
-class SCENEOUTLINER_API FSceneOutlinerGutter : public ISceneOutlinerColumn
+class FSceneOutlinerGutter : public ISceneOutlinerColumn
 {
 public:
 
 	/**	Constructor */
-	FSceneOutlinerGutter(ISceneOutliner& Outliner);
+	UE_API FSceneOutlinerGutter(ISceneOutliner& Outliner);
 
 	virtual ~FSceneOutlinerGutter() {}
 
@@ -32,17 +34,17 @@ public:
 
 	// -----------------------------------------
 	// ISceneOutlinerColumn Implementation
-	virtual FName GetColumnID() override;
+	UE_API virtual FName GetColumnID() override;
 
-	virtual SHeaderRow::FColumn::FArguments ConstructHeaderRowColumn() override;
+	UE_API virtual SHeaderRow::FColumn::FArguments ConstructHeaderRowColumn() override;
 
-	virtual const TSharedRef< SWidget > ConstructRowWidget(FSceneOutlinerTreeItemRef TreeItem, const STableRow<FSceneOutlinerTreeItemPtr>& Row) override;
+	UE_API virtual const TSharedRef< SWidget > ConstructRowWidget(FSceneOutlinerTreeItemRef TreeItem, const STableRow<FSceneOutlinerTreeItemPtr>& Row) override;
 
-	virtual void Tick(double InCurrentTime, float InDeltaTime) override;
+	UE_API virtual void Tick(double InCurrentTime, float InDeltaTime) override;
 
 	virtual bool SupportsSorting() const override { return true; }
 
-	virtual void SortItems(TArray<FSceneOutlinerTreeItemPtr>& RootItems, const EColumnSortMode::Type SortMode) const override;
+	UE_API virtual void SortItems(TArray<FSceneOutlinerTreeItemPtr>& RootItems, const EColumnSortMode::Type SortMode) const override;
 	// -----------------------------------------
 
 	/** Check whether the specified item is visible */
@@ -61,14 +63,14 @@ protected:
 };
 
 /** Widget responsible for managing the visibility for a single item */
-class SCENEOUTLINER_API SVisibilityWidget : public SImage
+class SVisibilityWidget : public SImage
 {
 public:
 	SLATE_BEGIN_ARGS(SVisibilityWidget) {}
 	SLATE_END_ARGS()
 
 		/** Construct this widget */
-		void Construct(const FArguments& InArgs, TWeakPtr<FSceneOutlinerGutter> InWeakColumn, TWeakPtr<ISceneOutliner> InWeakOutliner, TWeakPtr<ISceneOutlinerTreeItem> InWeakTreeItem, const STableRow<FSceneOutlinerTreeItemPtr>* InRow);
+		UE_API void Construct(const FArguments& InArgs, TWeakPtr<FSceneOutlinerGutter> InWeakColumn, TWeakPtr<ISceneOutliner> InWeakOutliner, TWeakPtr<ISceneOutlinerTreeItem> InWeakTreeItem, const STableRow<FSceneOutlinerTreeItemPtr>* InRow);
 
 protected:
 
@@ -76,43 +78,43 @@ protected:
 	virtual bool IsEnabled() const { return true; }
 
 	/** Get the brush for this widget */
-	virtual const FSlateBrush* GetBrush() const;
+	UE_API virtual const FSlateBrush* GetBrush() const;
 
 	/** Start a new drag/drop operation for this widget */
-	virtual FReply OnDragDetected(const FGeometry& MyGeometry, const FPointerEvent& MouseEvent) override;
+	UE_API virtual FReply OnDragDetected(const FGeometry& MyGeometry, const FPointerEvent& MouseEvent) override;
 
 	/** If a visibility drag drop operation has entered this widget, set its item to the new visibility state */
-	virtual void OnDragEnter(const FGeometry& MyGeometry, const FDragDropEvent& DragDropEvent) override;
+	UE_API virtual void OnDragEnter(const FGeometry& MyGeometry, const FDragDropEvent& DragDropEvent) override;
 
-	FReply HandleClick();
+	UE_API FReply HandleClick();
 
-	virtual FReply OnMouseButtonDoubleClick(const FGeometry& InMyGeometry, const FPointerEvent& InMouseEvent) override;
+	UE_API virtual FReply OnMouseButtonDoubleClick(const FGeometry& InMyGeometry, const FPointerEvent& InMouseEvent) override;
 
 	/** Called when the mouse button is pressed down on this widget */
-	virtual FReply OnMouseButtonDown(const FGeometry& MyGeometry, const FPointerEvent& MouseEvent) override;
+	UE_API virtual FReply OnMouseButtonDown(const FGeometry& MyGeometry, const FPointerEvent& MouseEvent) override;
 
 	/** Process a mouse up message */
-	virtual FReply OnMouseButtonUp(const FGeometry& MyGeometry, const FPointerEvent& MouseEvent) override;
+	UE_API virtual FReply OnMouseButtonUp(const FGeometry& MyGeometry, const FPointerEvent& MouseEvent) override;
 
 	/** Called when this widget had captured the mouse, but that capture has been revoked for some reason. */
-	virtual void OnMouseCaptureLost(const FCaptureLostEvent& CaptureLostEvent) override;
+	UE_API virtual void OnMouseCaptureLost(const FCaptureLostEvent& CaptureLostEvent) override;
 
 	/** Called when visibility change happens on an Item. */
-	virtual void OnSetItemVisibility(ISceneOutlinerTreeItem& Item, const bool bNewVisibility);
+	UE_API virtual void OnSetItemVisibility(ISceneOutlinerTreeItem& Item, const bool bNewVisibility);
 
 	/** Whether visibility change should propagate down to children. */
 	virtual bool ShouldPropagateVisibilityChangeOnChildren() const { return true; }
 
-	virtual FSlateColor GetForegroundColor() const;
+	UE_API virtual FSlateColor GetForegroundColor() const;
 
 	/** Check if the specified item is visible */
-	static bool IsVisible(const FSceneOutlinerTreeItemPtr& Item, const TSharedPtr<FSceneOutlinerGutter>& Column);
+	static UE_API bool IsVisible(const FSceneOutlinerTreeItemPtr& Item, const TSharedPtr<FSceneOutlinerGutter>& Column);
 
 	/** Check if our wrapped tree item is visible */
-	bool IsVisible() const;
+	UE_API bool IsVisible() const;
 
 	/** Set the item this widget is responsible for to be hidden or shown */
-	void SetIsVisible(const bool bVisible);
+	UE_API void SetIsVisible(const bool bVisible);
 
 	/** The tree item we relate to */
 	TWeakPtr<ISceneOutlinerTreeItem> WeakTreeItem;
@@ -134,4 +136,6 @@ protected:
 	const FSlateBrush* VisibleNotHoveredBrush;
 	const FSlateBrush* NotVisibleHoveredBrush;
 	const FSlateBrush* NotVisibleNotHoveredBrush;
-};
\ No newline at end of file
+};
+
+#undef UE_API
diff --git a/Engine/Source/Editor/SceneOutliner/Public/SceneOutlinerHelpers.h b/Engine/Source/Editor/SceneOutliner/Public/SceneOutlinerHelpers.h
--- a/Engine/Source/Editor/SceneOutliner/Public/SceneOutlinerHelpers.h
+++ b/Engine/Source/Editor/SceneOutliner/Public/SceneOutlinerHelpers.h
@@ -4,26 +4,30 @@
 
 #include "SceneOutlinerPublicTypes.h"
 
+#define UE_API SCENEOUTLINER_API
+
 namespace SceneOutliner
 {
 	// Class to hold common functionality needed by the Outliner and helpful to modules creating Outliner instances
-	class SCENEOUTLINER_API FSceneOutlinerHelpers
+	class FSceneOutlinerHelpers
 	{
 	public:
 		UE_DEPRECATED(5.5, "FSceneOutlinerHelpers::GetExternalPackageName has been deprecated implement/use ISceneOutlinerTreeItem::GetPackageName instead")
-		static FString GetExternalPackageName(const ISceneOutlinerTreeItem& TreeItem);
+		static UE_API FString GetExternalPackageName(const ISceneOutlinerTreeItem& TreeItem);
 		
 		UE_DEPRECATED(5.5, "FSceneOutlinerHelpers::GetExternalPackage has been deprecated implement/use ISceneOutlinerTreeItem::GetPackageName instead")
-		static UPackage* GetExternalPackage(const ISceneOutlinerTreeItem& TreeItem);
+		static UE_API UPackage* GetExternalPackage(const ISceneOutlinerTreeItem& TreeItem);
 		
-		static TSharedPtr<SWidget> GetClassHyperlink(UObject* InObject);
+		static UE_API TSharedPtr<SWidget> GetClassHyperlink(UObject* InObject);
 		
-		static void PopulateExtraSearchStrings(const ISceneOutlinerTreeItem& TreeItem, TArray< FString >& OutSearchStrings);
+		static UE_API void PopulateExtraSearchStrings(const ISceneOutlinerTreeItem& TreeItem, TArray< FString >& OutSearchStrings);
 
-		static bool ValidateFolderName(const FFolder& InFolder, UWorld* World, const FText& InLabel, FText& OutErrorMessage);
+		static UE_API bool ValidateFolderName(const FFolder& InFolder, UWorld* World, const FText& InLabel, FText& OutErrorMessage);
 
-		static void RenameFolder(const FFolder& InFolder, const FText& NewFolderName, UWorld* World);
+		static UE_API void RenameFolder(const FFolder& InFolder, const FText& NewFolderName, UWorld* World);
 		
-		static bool IsFolderCurrent(const FFolder& InFolder, UWorld* World);
+		static UE_API bool IsFolderCurrent(const FFolder& InFolder, UWorld* World);
 	};
 };
+
+#undef UE_API
diff --git a/Engine/Source/Editor/SceneOutliner/Public/SceneOutlinerMenuContext.h b/Engine/Source/Editor/SceneOutliner/Public/SceneOutlinerMenuContext.h
--- a/Engine/Source/Editor/SceneOutliner/Public/SceneOutlinerMenuContext.h
+++ b/Engine/Source/Editor/SceneOutliner/Public/SceneOutlinerMenuContext.h
@@ -9,8 +9,8 @@
 
 class SSceneOutliner;
 
-UCLASS()
-class SCENEOUTLINER_API USceneOutlinerMenuContext : public UObject
+UCLASS(MinimalAPI)
+class USceneOutlinerMenuContext : public UObject
 {
 	GENERATED_BODY()
 public:
diff --git a/Engine/Source/Editor/SceneOutliner/Public/SceneOutlinerPublicTypes.h b/Engine/Source/Editor/SceneOutliner/Public/SceneOutlinerPublicTypes.h
--- a/Engine/Source/Editor/SceneOutliner/Public/SceneOutlinerPublicTypes.h
+++ b/Engine/Source/Editor/SceneOutliner/Public/SceneOutlinerPublicTypes.h
@@ -8,6 +8,8 @@
 
 #include "SceneOutlinerFilters.h"
 
+#define UE_API SCENEOUTLINER_API
+
 #define LOCTEXT_NAMESPACE "SceneOutlinerPublicTypes"
 
 class FExtender;
@@ -106,7 +108,7 @@ struct FSceneOutlinerColumnInfo
 };
 
 /** Settings for the scene outliner which can be quieried publicly */
-struct SCENEOUTLINER_API FSharedSceneOutlinerData
+struct FSharedSceneOutlinerData
 {
 	/**	Invoked whenever the user attempts to delete an actor from within a Scene Outliner in the actor browsing mode */
 	FCustomSceneOutlinerDeleteDelegate CustomDelete;
@@ -136,7 +138,7 @@ public:
 	{}
 
 	/** Set up a default array of columns for this outliner */
-	void UseDefaultColumns();
+	UE_API void UseDefaultColumns();
 };
 
 /* Settings for the Filter Bar attached to the Scene Outliner. Can be specified through FSceneOutlinerInitializationOptions */
@@ -225,15 +227,17 @@ struct FSceneOutlinerDefaultTreeItemMetrics
 };
 
 /** A struct which gets, and caches the visibility of a tree item */
-struct SCENEOUTLINER_API FSceneOutlinerVisibilityCache
+struct FSceneOutlinerVisibilityCache
 {
 	/** Map of tree item to visibility */
 	mutable TMap<const ISceneOutlinerTreeItem*, bool> VisibilityInfo;
 
 	/** Get an item's visibility based on its children */
-	bool RecurseChildren(const ISceneOutlinerTreeItem& Item) const;
+	UE_API bool RecurseChildren(const ISceneOutlinerTreeItem& Item) const;
 
-	bool GetVisibility(const ISceneOutlinerTreeItem& Item) const;
+	UE_API bool GetVisibility(const ISceneOutlinerTreeItem& Item) const;
 };
 
 #undef LOCTEXT_NAMESPACE
+
+#undef UE_API
diff --git a/Engine/Source/Editor/SceneOutliner/Public/SceneOutlinerSourceControlColumn.h b/Engine/Source/Editor/SceneOutliner/Public/SceneOutlinerSourceControlColumn.h
--- a/Engine/Source/Editor/SceneOutliner/Public/SceneOutlinerSourceControlColumn.h
+++ b/Engine/Source/Editor/SceneOutliner/Public/SceneOutlinerSourceControlColumn.h
@@ -9,10 +9,12 @@
 #include "SceneOutlinerPublicTypes.h"
 #include "ISceneOutlinerColumn.h"
 
+#define UE_API SCENEOUTLINER_API
+
 template<typename ItemType> class STableRow;
 
 /** A column for the SceneOutliner that displays the SCC Information */
-class SCENEOUTLINER_API FSceneOutlinerSourceControlColumn : public ISceneOutlinerColumn
+class FSceneOutlinerSourceControlColumn : public ISceneOutlinerColumn
 {
 
 public:
@@ -27,18 +29,20 @@ public:
 	//////////////////////////////////////////////////////////////////////////
 	// Begin ISceneOutlinerColumn Implementation
 
-	virtual FName GetColumnID() override;
+	UE_API virtual FName GetColumnID() override;
 
-	virtual SHeaderRow::FColumn::FArguments ConstructHeaderRowColumn() override;
+	UE_API virtual SHeaderRow::FColumn::FArguments ConstructHeaderRowColumn() override;
 
-	virtual const TSharedRef<SWidget> ConstructRowWidget(FSceneOutlinerTreeItemRef TreeItem, const STableRow<FSceneOutlinerTreeItemPtr>& Row) override;
+	UE_API virtual const TSharedRef<SWidget> ConstructRowWidget(FSceneOutlinerTreeItemRef TreeItem, const STableRow<FSceneOutlinerTreeItemPtr>& Row) override;
 
 	virtual bool SupportsSorting() const override { return false; }
 
 	// End ISceneOutlinerColumn Implementation
 	//////////////////////////////////////////////////////////////////////////
 private:
-	const FSlateBrush* GetHeaderIconBadge() const;
+	UE_API const FSlateBrush* GetHeaderIconBadge() const;
 
 	TWeakPtr<ISceneOutliner> WeakSceneOutliner;
 };
+
+#undef UE_API
diff --git a/Engine/Source/Editor/SceneOutliner/Public/SceneOutlinerStandaloneTypes.h b/Engine/Source/Editor/SceneOutliner/Public/SceneOutlinerStandaloneTypes.h
--- a/Engine/Source/Editor/SceneOutliner/Public/SceneOutlinerStandaloneTypes.h
+++ b/Engine/Source/Editor/SceneOutliner/Public/SceneOutlinerStandaloneTypes.h
@@ -9,6 +9,8 @@
 #include "Folder.h"
 #include "SceneOutlinerFwd.h"
 
+#define UE_API SCENEOUTLINER_API
+
 struct FFolderKey
 {
 	FFolderKey(const FName& InPath, const FFolder::FRootObject& InRootObject)
@@ -179,7 +181,7 @@ private:
 	mutable uint8 Data[MaxSize];
 };
 	
-struct SCENEOUTLINER_API FSceneOutlinerTreeItemType
+struct FSceneOutlinerTreeItemType
 {
 public:
 	explicit FSceneOutlinerTreeItemType(const FSceneOutlinerTreeItemType* Parent = nullptr) : ID(++NextUniqueID), ParentType(Parent) {}
@@ -195,19 +197,19 @@ public:
 	}
 
 private:
-	static uint32 NextUniqueID;
+	static UE_API uint32 NextUniqueID;
 	uint32 ID;
 	const FSceneOutlinerTreeItemType* ParentType;
 };
 
-struct SCENEOUTLINER_API FSceneOutlinerCommonLabelData
+struct FSceneOutlinerCommonLabelData
 {
 	TWeakPtr<ISceneOutliner> WeakSceneOutliner;
-	static const FLinearColor DarkColor;
+	static UE_API const FLinearColor DarkColor;
 
-	TOptional<FLinearColor> GetForegroundColor(const ISceneOutlinerTreeItem& TreeItem) const;
+	UE_API TOptional<FLinearColor> GetForegroundColor(const ISceneOutlinerTreeItem& TreeItem) const;
 
-	bool CanExecuteRenameRequest(const ISceneOutlinerTreeItem& Item) const;
+	UE_API bool CanExecuteRenameRequest(const ISceneOutlinerTreeItem& Item) const;
 };
 
 /**
@@ -236,3 +238,5 @@ struct FSceneOutlinerHierarchyChangedData
 	/** Actions to apply to items */
 	uint8 ItemActions = 0;
 };
+
+#undef UE_API
diff --git a/Engine/Source/Editor/SceneOutliner/Public/SceneOutlinerTextInfoColumn.h b/Engine/Source/Editor/SceneOutliner/Public/SceneOutlinerTextInfoColumn.h
--- a/Engine/Source/Editor/SceneOutliner/Public/SceneOutlinerTextInfoColumn.h
+++ b/Engine/Source/Editor/SceneOutliner/Public/SceneOutlinerTextInfoColumn.h
@@ -12,6 +12,8 @@
 #include "ISceneOutlinerColumn.h"
 #include "Misc/CoreDelegates.h"
 
+#define UE_API SCENEOUTLINER_API
+
 class ISceneOutliner;
 
 /** A delegate used to factory a new TextInfo Column */
@@ -22,7 +24,7 @@ DECLARE_DELEGATE_RetVal_OneParam(FString, FGetTextForItem, const ISceneOutlinerT
 * A custom column for the SceneOutliner which is capable of displaying any text based information related to items
 * Note: Only has support for text, any other requirements should have a custom column made
 */
-class SCENEOUTLINER_API FTextInfoColumn : public ISceneOutlinerColumn
+class FTextInfoColumn : public ISceneOutlinerColumn
 {
 
 	public:
@@ -30,34 +32,34 @@ class SCENEOUTLINER_API FTextInfoColumn : public ISceneOutlinerColumn
 		/**
 		 *	Constructor
 		 */
-		FTextInfoColumn(ISceneOutliner& Outliner, const FName InColumnName, const FGetTextForItem& InGetTextForItem, const FText InColumnToolTip);
+		UE_API FTextInfoColumn(ISceneOutliner& Outliner, const FName InColumnName, const FGetTextForItem& InGetTextForItem, const FText InColumnToolTip);
 
 		virtual ~FTextInfoColumn() {}
 
 		// Factory function to create a TextInfoColumn (This column is not registered by default because of it's requirements and therefore requires a factory!)
-		static TSharedRef<ISceneOutlinerColumn> CreateTextInfoColumn(ISceneOutliner& Outliner, const FName InColumnName, const FGetTextForItem InGetTextForItem, const FText InColumnToolTip);
+		static UE_API TSharedRef<ISceneOutlinerColumn> CreateTextInfoColumn(ISceneOutliner& Outliner, const FName InColumnName, const FGetTextForItem InGetTextForItem, const FText InColumnToolTip);
 
 		//////////////////////////////////////////////////////////////////////////
 		// Begin ISceneOutlinerColumn Implementation
 
-		virtual FName GetColumnID() override;
+		UE_API virtual FName GetColumnID() override;
 
-		virtual SHeaderRow::FColumn::FArguments ConstructHeaderRowColumn() override;
+		UE_API virtual SHeaderRow::FColumn::FArguments ConstructHeaderRowColumn() override;
 
-		virtual const TSharedRef< SWidget > ConstructRowWidget(FSceneOutlinerTreeItemRef TreeItem, const STableRow<FSceneOutlinerTreeItemPtr>& Row) override;
+		UE_API virtual const TSharedRef< SWidget > ConstructRowWidget(FSceneOutlinerTreeItemRef TreeItem, const STableRow<FSceneOutlinerTreeItemPtr>& Row) override;
 
-		virtual void PopulateSearchStrings(const ISceneOutlinerTreeItem& Item, TArray< FString >& OutSearchStrings) const override;
+		UE_API virtual void PopulateSearchStrings(const ISceneOutlinerTreeItem& Item, TArray< FString >& OutSearchStrings) const override;
 
-		virtual bool SupportsSorting() const override;
+		UE_API virtual bool SupportsSorting() const override;
 
-		virtual void SortItems(TArray<FSceneOutlinerTreeItemPtr>& RootItems, const EColumnSortMode::Type SortMode) const override;
+		UE_API virtual void SortItems(TArray<FSceneOutlinerTreeItemPtr>& RootItems, const EColumnSortMode::Type SortMode) const override;
 
 		// End ISceneOutlinerColumn Implementation
 		//////////////////////////////////////////////////////////////////////////
 
 	private:
 
-		FText GetInfoForItem(TWeakPtr<ISceneOutlinerTreeItem> TreeItem) const;
+		UE_API FText GetInfoForItem(TWeakPtr<ISceneOutlinerTreeItem> TreeItem) const;
 
 		FName ColumnName;
 
@@ -69,3 +71,5 @@ class SCENEOUTLINER_API FTextInfoColumn : public ISceneOutlinerColumn
 		/* Delegate Called to Get Display Text for an item in the column*/
 		FGetTextForItem GetTextForItem;
 };
+
+#undef UE_API
diff --git a/Engine/Source/Editor/SceneOutliner/Public/WorldTreeItem.h b/Engine/Source/Editor/SceneOutliner/Public/WorldTreeItem.h
--- a/Engine/Source/Editor/SceneOutliner/Public/WorldTreeItem.h
+++ b/Engine/Source/Editor/SceneOutliner/Public/WorldTreeItem.h
@@ -7,6 +7,8 @@
 #include "ISceneOutlinerTreeItem.h"
 #include "UObject/ObjectKey.h"
 
+#define UE_API SCENEOUTLINER_API
+
 class UToolMenu;
 
 namespace SceneOutliner
@@ -16,7 +18,7 @@ namespace SceneOutliner
 }
 
 /** A tree item that represents an entire world */
-struct SCENEOUTLINER_API FWorldTreeItem : ISceneOutlinerTreeItem
+struct FWorldTreeItem : ISceneOutlinerTreeItem
 {
 public:
 	DECLARE_DELEGATE_RetVal_OneParam(bool, FFilterPredicate, const UWorld*);
@@ -39,27 +41,29 @@ public:
 	const FObjectKey ID;
 
 	/** Static type identifier for this tree item class */
-	static const FSceneOutlinerTreeItemType Type;
+	static UE_API const FSceneOutlinerTreeItemType Type;
 
 	/** Construct this item from a world */
-	FWorldTreeItem(UWorld* InWorld);
-	FWorldTreeItem(TWeakObjectPtr<UWorld> InWorld);
+	UE_API FWorldTreeItem(UWorld* InWorld);
+	UE_API FWorldTreeItem(TWeakObjectPtr<UWorld> InWorld);
 
 	/* Begin ISceneOutlinerTreeItem Implementation */
 	virtual bool IsValid() const override { return World.IsValid(); }
-	virtual FSceneOutlinerTreeItemID GetID() const override;
-	virtual FString GetDisplayString() const override;
-	virtual bool CanInteract() const override;
-	virtual void GenerateContextMenu(UToolMenu* Menu, SSceneOutliner& Outliner) override;
-	virtual TSharedRef<SWidget> GenerateLabelWidget(ISceneOutliner& Outliner, const STableRow<FSceneOutlinerTreeItemPtr>& InRow) override;
+	UE_API virtual FSceneOutlinerTreeItemID GetID() const override;
+	UE_API virtual FString GetDisplayString() const override;
+	UE_API virtual bool CanInteract() const override;
+	UE_API virtual void GenerateContextMenu(UToolMenu* Menu, SSceneOutliner& Outliner) override;
+	UE_API virtual TSharedRef<SWidget> GenerateLabelWidget(ISceneOutliner& Outliner, const STableRow<FSceneOutlinerTreeItemPtr>& InRow) override;
 	/* End ISceneOutlinerTreeItem Implementation */
 
 	/** Open the world settings for the contained world */
-	void OpenWorldSettings() const;
+	UE_API void OpenWorldSettings() const;
 	/** Browse to world asset in content browser */
-	void BrowseToAsset() const;
-	bool CanBrowseToAsset() const;
+	UE_API void BrowseToAsset() const;
+	UE_API bool CanBrowseToAsset() const;
 
 	/** Get just the name of the world, for tooltip use */
-	FString GetWorldName() const;
+	UE_API FString GetWorldName() const;
 };
+
+#undef UE_API
diff --git a/Engine/Source/Editor/ScriptableEditorWidgets/Public/Components/DetailsView.h b/Engine/Source/Editor/ScriptableEditorWidgets/Public/Components/DetailsView.h
--- a/Engine/Source/Editor/ScriptableEditorWidgets/Public/Components/DetailsView.h
+++ b/Engine/Source/Editor/ScriptableEditorWidgets/Public/Components/DetailsView.h
@@ -13,6 +13,8 @@
 
 #include "DetailsView.generated.h"
 
+#define UE_API SCRIPTABLEEDITORWIDGETS_API
+
 class FProperty;
 class IDetailsView;
 class UObject;
@@ -22,8 +24,8 @@ struct FPropertyChangedEvent;
 /**
  * The details view allows you to display the value of an object properties.
  */
-UCLASS()
-class SCRIPTABLEEDITORWIDGETS_API UDetailsView : public UPropertyViewBase, public FNotifyHook
+UCLASS(MinimalAPI)
+class UDetailsView : public UPropertyViewBase, public FNotifyHook
 {
 	GENERATED_BODY()
 
@@ -74,28 +76,28 @@ public:
 
 	//~ UPropertyViewBase interface
 protected:
-	virtual void BuildContentWidget() override;
-	virtual void OnObjectChanged() override;
+	UE_API virtual void BuildContentWidget() override;
+	UE_API virtual void OnObjectChanged() override;
 	//~ End of UPropertyViewBase interface
 
 public:
 	//~ UWidget interface
-	virtual void ReleaseSlateResources(bool bReleaseChildren) override;
+	UE_API virtual void ReleaseSlateResources(bool bReleaseChildren) override;
 	// End of UWidget interface
 
 	//~ FNotifyHook interface
-	virtual void NotifyPostChange(const FPropertyChangedEvent& PropertyChangedEvent, FProperty* PropertyThatChanged) override;
+	UE_API virtual void NotifyPostChange(const FPropertyChangedEvent& PropertyChangedEvent, FProperty* PropertyThatChanged) override;
 	//~ End of FNotifyHook interface
 
 	//~ UObject interface
-	virtual void PostEditChangeProperty(FPropertyChangedEvent& PropertyChangedEvent) override;
+	UE_API virtual void PostEditChangeProperty(FPropertyChangedEvent& PropertyChangedEvent) override;
 	//~ End of UObject interface
 
 protected:
-	void ToggleShowingOnlyAllowedProperties();
-	virtual bool IsRowVisibilityFiltered() const;
-	virtual bool GetIsPropertyVisible(const FPropertyAndParent& PropertyAndParent) const;
-	virtual bool GetIsRowVisible(FName InRowName, FName InParentName) const;
+	UE_API void ToggleShowingOnlyAllowedProperties();
+	UE_API virtual bool IsRowVisibilityFiltered() const;
+	UE_API virtual bool GetIsPropertyVisible(const FPropertyAndParent& PropertyAndParent) const;
+	UE_API virtual bool GetIsRowVisible(FName InRowName, FName InParentName) const;
 
 private:
 	TSharedPtr<IDetailsView> DetailViewWidget;
@@ -104,3 +106,5 @@ private:
 	UPROPERTY()
 	bool bShowOnlyAllowed = true;
 };
+
+#undef UE_API
diff --git a/Engine/Source/Editor/ScriptableEditorWidgets/Public/Components/PropertyViewBase.h b/Engine/Source/Editor/ScriptableEditorWidgets/Public/Components/PropertyViewBase.h
--- a/Engine/Source/Editor/ScriptableEditorWidgets/Public/Components/PropertyViewBase.h
+++ b/Engine/Source/Editor/ScriptableEditorWidgets/Public/Components/PropertyViewBase.h
@@ -9,6 +9,8 @@
 #include "UObject/SoftObjectPtr.h"
 #include "PropertyViewBase.generated.h"
 
+#define UE_API SCRIPTABLEEDITORWIDGETS_API
+
 class SBorder;
 
 /** Sets a delegate called when the property value changes */
@@ -18,8 +20,8 @@ DECLARE_DYNAMIC_MULTICAST_DELEGATE_OneParam(FOnPropertyValueChanged, FName, Prop
 /**
  * Base of property view allows you to display the value of an object properties.
  */
-UCLASS(Abstract)
-class SCRIPTABLEEDITORWIDGETS_API UPropertyViewBase : public UWidget
+UCLASS(MinimalAPI, Abstract)
+class UPropertyViewBase : public UWidget
 {
 	GENERATED_BODY()
 
@@ -42,32 +44,34 @@ protected:
 
 public:
 	UFUNCTION(BlueprintCallable, Category = "View")
-	UObject* GetObject() const;
+	UE_API UObject* GetObject() const;
 
 	UFUNCTION(BlueprintCallable, Category = "View")
-	void SetObject(UObject* NewObject);
+	UE_API void SetObject(UObject* NewObject);
 
 protected:
 	virtual void BuildContentWidget() PURE_VIRTUAL(UPropertyViewBase::BuildContentWidget, );
 	virtual void OnObjectChanged() { }
 	TSharedPtr<SBorder> GetDisplayWidget() const { return DisplayedWidget; }
-	void OnPropertyChangedBroadcast(FName PropertyName);
+	UE_API void OnPropertyChangedBroadcast(FName PropertyName);
 
 	//~ UWidget interface
 public:
-	virtual void ReleaseSlateResources(bool bReleaseChildren) override;
-	virtual const FText GetPaletteCategory() override;
+	UE_API virtual void ReleaseSlateResources(bool bReleaseChildren) override;
+	UE_API virtual const FText GetPaletteCategory() override;
 
 protected:
-	virtual TSharedRef<SWidget> RebuildWidget() override;
+	UE_API virtual TSharedRef<SWidget> RebuildWidget() override;
 	//~ End of UWidget interface
 
 	//~ UObject interface
 public:
-	virtual void PostLoad() override;
-	virtual void PostEditChangeProperty(FPropertyChangedEvent& PropertyChangedEvent) override;
+	UE_API virtual void PostLoad() override;
+	UE_API virtual void PostEditChangeProperty(FPropertyChangedEvent& PropertyChangedEvent) override;
 	//~ End of UObject interface
 
 private:
 	TSharedPtr<SBorder> DisplayedWidget;
 };
+
+#undef UE_API
diff --git a/Engine/Source/Editor/ScriptableEditorWidgets/Public/Components/SinglePropertyView.h b/Engine/Source/Editor/ScriptableEditorWidgets/Public/Components/SinglePropertyView.h
--- a/Engine/Source/Editor/ScriptableEditorWidgets/Public/Components/SinglePropertyView.h
+++ b/Engine/Source/Editor/ScriptableEditorWidgets/Public/Components/SinglePropertyView.h
@@ -10,6 +10,8 @@
 
 #include "SinglePropertyView.generated.h"
 
+#define UE_API SCRIPTABLEEDITORWIDGETS_API
+
 class ISinglePropertyView;
 class UObject;
 struct FFrame;
@@ -18,8 +20,8 @@ struct FPropertyChangedEvent;
 /**
  * The single property view allows you to display the value of an object's property.
  */
-UCLASS()
-class SCRIPTABLEEDITORWIDGETS_API USinglePropertyView : public UPropertyViewBase
+UCLASS(MinimalAPI)
+class USinglePropertyView : public UPropertyViewBase
 {
 	GENERATED_BODY()
 
@@ -34,36 +36,38 @@ private:
 
 public:
 	UFUNCTION(BlueprintCallable, Category = "View")
-	FName GetPropertyName() const;
+	UE_API FName GetPropertyName() const;
 
 	UFUNCTION(BlueprintCallable, Category = "View")
-	void SetPropertyName(FName NewPropertyName);
+	UE_API void SetPropertyName(FName NewPropertyName);
 
 	UFUNCTION(BlueprintCallable, Category = "View")
-	FText GetNameOverride() const;
+	UE_API FText GetNameOverride() const;
 
 	UFUNCTION(BlueprintCallable, Category = "View")
-	void SetNameOverride(FText NewPropertyName);
+	UE_API void SetNameOverride(FText NewPropertyName);
 
 private:
-	void InternalSinglePropertyChanged();
+	UE_API void InternalSinglePropertyChanged();
 
 	//~ UPropertyViewBase interface
 protected:
-	virtual void BuildContentWidget() override;
-	virtual void OnObjectChanged() override;
+	UE_API virtual void BuildContentWidget() override;
+	UE_API virtual void OnObjectChanged() override;
 	//~ End of UPropertyViewBase interface
 
 	//~ UWidget interface
 public:
-	virtual void ReleaseSlateResources(bool bReleaseChildren) override;
+	UE_API virtual void ReleaseSlateResources(bool bReleaseChildren) override;
 	//~ End of UWidget interface
 
 	// UObject interface
 public:
-	virtual void PostEditChangeProperty(FPropertyChangedEvent& PropertyChangedEvent) override;
+	UE_API virtual void PostEditChangeProperty(FPropertyChangedEvent& PropertyChangedEvent) override;
 	// End of UObject interface
 
 private:
 	TSharedPtr<ISinglePropertyView> SinglePropertyViewWidget;
 };
+
+#undef UE_API
diff --git a/Engine/Source/Editor/SequenceRecorder/Public/ActorRecordingSettings.h b/Engine/Source/Editor/SequenceRecorder/Public/ActorRecordingSettings.h
--- a/Engine/Source/Editor/SequenceRecorder/Public/ActorRecordingSettings.h
+++ b/Engine/Source/Editor/SequenceRecorder/Public/ActorRecordingSettings.h
@@ -11,18 +11,20 @@
 
 #include "ActorRecordingSettings.generated.h"
 
+#define UE_API SEQUENCERECORDER_API
+
 class UObject;
 
 USTRUCT()
-struct SEQUENCERECORDER_API FActorRecordingSettings
+struct FActorRecordingSettings
 {
 	GENERATED_BODY()
 
-	FActorRecordingSettings();
+	UE_API FActorRecordingSettings();
 
-	FActorRecordingSettings(class UObject* InOuter);
+	UE_API FActorRecordingSettings(class UObject* InOuter);
 
-	void CreateSettingsObjectsFromFactory();
+	UE_API void CreateSettingsObjectsFromFactory();
 
 	template <typename SettingsType>
 	SettingsType* GetSettingsObject() const
@@ -46,3 +48,5 @@ private:
 	/** An optional outer that settings objects should be created with. */
 	TWeakObjectPtr<UObject> Outer;
 };
+
+#undef UE_API
diff --git a/Engine/Source/Editor/SequenceRecorder/Public/AnimationRecorder.h b/Engine/Source/Editor/SequenceRecorder/Public/AnimationRecorder.h
--- a/Engine/Source/Editor/SequenceRecorder/Public/AnimationRecorder.h
+++ b/Engine/Source/Editor/SequenceRecorder/Public/AnimationRecorder.h
@@ -14,6 +14,8 @@
 #include "Animation/AnimTypes.h"
 #include "AnimationRecorder.generated.h"
 
+#define UE_API SEQUENCERECORDER_API
+
 class UAnimBoneCompressionSettings;
 class UAnimNotify;
 class UAnimNotifyState;
@@ -48,7 +50,7 @@ struct FTimecodeBoneMethod
 	FName BoneName;
 };
 
-struct SEQUENCERECORDER_API FProcessRecordedTimeParams
+struct FProcessRecordedTimeParams
 {
 	FString HoursName;
 	FString MinutesName;
@@ -66,7 +68,7 @@ struct SEQUENCERECORDER_API FProcessRecordedTimeParams
 // FAnimationRecorder
 
 // records the mesh pose to animation input
-struct SEQUENCERECORDER_API FAnimationRecorder : public FGCObject
+struct FAnimationRecorder : public FGCObject
 {
 private:
 	/** Frame count used to signal an unbounded animation */
@@ -124,17 +126,17 @@ private:
 	/** Currently recording notify events that have duration */
 	TArray<FRecordedAnimNotify> RecordingAnimNotifies;
 
-	static float DefaultSampleRate;
+	static UE_API float DefaultSampleRate;
 
 	/** Array of times recorded */
 	TArray<FQualifiedFrameTime> RecordedTimes;
 
 public:
-	FAnimationRecorder();
-	~FAnimationRecorder();
+	UE_API FAnimationRecorder();
+	UE_API ~FAnimationRecorder();
 
 	// FGCObject interface start
-	virtual void AddReferencedObjects(FReferenceCollector& Collector) override;
+	UE_API virtual void AddReferencedObjects(FReferenceCollector& Collector) override;
 	virtual FString GetReferencerName() const override
 	{
 		return TEXT("FAnimationRecorder");
@@ -142,31 +144,31 @@ public:
 	// FGCObject interface end
 
 	/** Starts recording an animation. Prompts for asset path and name via dialog if none provided */
-	bool TriggerRecordAnimation(USkeletalMeshComponent* Component);
-	bool TriggerRecordAnimation(USkeletalMeshComponent* Component, const FString& AssetPath, const FString& AssetName);
+	UE_API bool TriggerRecordAnimation(USkeletalMeshComponent* Component);
+	UE_API bool TriggerRecordAnimation(USkeletalMeshComponent* Component, const FString& AssetPath, const FString& AssetName);
 
-	void StartRecord(USkeletalMeshComponent* Component, UAnimSequence* InAnimationObject);
-	UAnimSequence* StopRecord(bool bShowMessage);
-	void UpdateRecord(USkeletalMeshComponent* Component, float DeltaTime);
+	UE_API void StartRecord(USkeletalMeshComponent* Component, UAnimSequence* InAnimationObject);
+	UE_API UAnimSequence* StopRecord(bool bShowMessage);
+	UE_API void UpdateRecord(USkeletalMeshComponent* Component, float DeltaTime);
 	UAnimSequence* GetAnimationObject() const { return AnimationObject; }
 	bool InRecording() const { return AnimationObject != nullptr; }
 	double GetTimeRecorded() const { return TimePassed; }
 
 	/** Sets a new sample rate & max length for this recorder. Don't call while recording. */
-	void SetSampleRateAndLength(FFrameRate SampleFrameRate, float LengthInSeconds);
+	UE_API void SetSampleRateAndLength(FFrameRate SampleFrameRate, float LengthInSeconds);
 
-	bool SetAnimCompressionScheme(UAnimBoneCompressionSettings* Settings);
+	UE_API bool SetAnimCompressionScheme(UAnimBoneCompressionSettings* Settings);
 
 	/** Sets the function that should be used to get the current timecode for the current frame. */
-	void SetCurrentFrameTimeGetter(TFunction<TOptional<FQualifiedFrameTime>()> InGetCurrentTimeFunction);
+	UE_API void SetCurrentFrameTimeGetter(TFunction<TOptional<FQualifiedFrameTime>()> InGetCurrentTimeFunction);
 
 	const FTransform& GetInitialRootTransform() const { return InitialRootTransform; }
 
 	UE_DEPRECATED(5.5, "Use the ProcessRecordedTimes method that takes a FProcessRecordedTimeParams struct.")
-	void ProcessRecordedTimes(UAnimSequence* AnimSequence, USkeletalMeshComponent* SkeletalMeshComponent, const FString& HoursName, const FString& MinutesName, const FString& SecondsName, const FString& FramesName, const FString& SubFramesName, const FString& SlateName, const FString& Slate, const FTimecodeBoneMethod& TimecodeBoneMethod);
+	UE_API void ProcessRecordedTimes(UAnimSequence* AnimSequence, USkeletalMeshComponent* SkeletalMeshComponent, const FString& HoursName, const FString& MinutesName, const FString& SecondsName, const FString& FramesName, const FString& SubFramesName, const FString& SlateName, const FString& Slate, const FTimecodeBoneMethod& TimecodeBoneMethod);
 
 	/** Process any time data captured and apply it to the bones on the given SkeletalMeshComponent. */
-	void ProcessRecordedTimes(UAnimSequence* AnimSequence, USkeletalMeshComponent* SkeletalMeshComponent, const FTimecodeBoneMethod& TimecodeBoneMethod, const FProcessRecordedTimeParams& TimecodeInfo);
+	UE_API void ProcessRecordedTimes(UAnimSequence* AnimSequence, USkeletalMeshComponent* SkeletalMeshComponent, const FTimecodeBoneMethod& TimecodeBoneMethod, const FProcessRecordedTimeParams& TimecodeInfo);
 
 	/** If true, it will record root to include LocalToWorld */
 	uint8 bRecordLocalToWorld :1;
@@ -200,16 +202,16 @@ public:
 	bool bTransactRecording;
 public:
 	/** Helper function to get space bases depending on leader pose component */
-	static void GetBoneTransforms(USkeletalMeshComponent* Component, TArray<FTransform>& BoneTransforms);
+	static UE_API void GetBoneTransforms(USkeletalMeshComponent* Component, TArray<FTransform>& BoneTransforms);
 
 private:
-	bool Record(USkeletalMeshComponent* Component, FTransform const& ComponentToWorld, const TArray<FTransform>& SpacesBases, const FBlendedHeapCurve& AnimationCurves, int32 FrameToAdd);
+	UE_API bool Record(USkeletalMeshComponent* Component, FTransform const& ComponentToWorld, const TArray<FTransform>& SpacesBases, const FBlendedHeapCurve& AnimationCurves, int32 FrameToAdd);
 
-	void RecordNotifies(USkeletalMeshComponent* Component, const TArray<FAnimNotifyEventReference>& AnimNotifies, float DeltaTime, float RecordTime);
+	UE_API void RecordNotifies(USkeletalMeshComponent* Component, const TArray<FAnimNotifyEventReference>& AnimNotifies, float DeltaTime, float RecordTime);
 
-	void ProcessNotifies();
+	UE_API void ProcessNotifies();
 
-	bool ShouldSkipName(const FName& InName) const;
+	UE_API bool ShouldSkipName(const FName& InName) const;
 
 	TArray<FBlendedHeapCurve> RecordedCurves;
 	TArray<FRawAnimSequenceTrack> RawTracks;
@@ -218,28 +220,28 @@ private:
 //////////////////////////////////////////////////////////////////////////
 // FAnimRecorderInstance
 
-struct SEQUENCERECORDER_API FAnimRecorderInstance
+struct FAnimRecorderInstance
 {
 public:
-	FAnimRecorderInstance();
-	~FAnimRecorderInstance();
+	UE_API FAnimRecorderInstance();
+	UE_API ~FAnimRecorderInstance();
 
-	void Init(USkeletalMeshComponent* InComponent, const FString& InAssetPath, const FString& InAssetName, const FAnimationRecordingSettings& InSettings);
+	UE_API void Init(USkeletalMeshComponent* InComponent, const FString& InAssetPath, const FString& InAssetName, const FAnimationRecordingSettings& InSettings);
 
-	void Init(USkeletalMeshComponent* InComponent, UAnimSequence* InSequence, FAnimationSerializer *InAnimationSerializer, const FAnimationRecordingSettings& InSettings);
+	UE_API void Init(USkeletalMeshComponent* InComponent, UAnimSequence* InSequence, FAnimationSerializer *InAnimationSerializer, const FAnimationRecordingSettings& InSettings);
 
-	bool BeginRecording();
-	void Update(float DeltaTime);
-	void FinishRecording(bool bShowMessage = true);
+	UE_API bool BeginRecording();
+	UE_API void Update(float DeltaTime);
+	UE_API void FinishRecording(bool bShowMessage = true);
 
 	UE_DEPRECATED(5.5, "Use the ProcessRecordedTimes method that takes a FProcessRecordedTimeParams struct.")
-	void ProcessRecordedTimes(UAnimSequence* AnimSequence, USkeletalMeshComponent* SkeletalMeshComponent, const FString& HoursName, const FString& MinutesName, const FString& SecondsName, const FString& FramesName, const FString& SubFramesName, const FString& SlateName, const FString& Slate, const FTimecodeBoneMethod& TimecodeBoneMethod);
+	UE_API void ProcessRecordedTimes(UAnimSequence* AnimSequence, USkeletalMeshComponent* SkeletalMeshComponent, const FString& HoursName, const FString& MinutesName, const FString& SecondsName, const FString& FramesName, const FString& SubFramesName, const FString& SlateName, const FString& Slate, const FTimecodeBoneMethod& TimecodeBoneMethod);
 
 	/** Process any time data captured and apply it to the bones on the given SkeletalMeshComponent. */
-	void ProcessRecordedTimes(UAnimSequence* AnimSequence, USkeletalMeshComponent* SkeletalMeshComponent, const FTimecodeBoneMethod& TimecodeBoneMethod, const FProcessRecordedTimeParams& TimecodeInfo);
+	UE_API void ProcessRecordedTimes(UAnimSequence* AnimSequence, USkeletalMeshComponent* SkeletalMeshComponent, const FTimecodeBoneMethod& TimecodeBoneMethod, const FProcessRecordedTimeParams& TimecodeInfo);
 
 private:
-	void InitInternal(USkeletalMeshComponent* InComponent, const FAnimationRecordingSettings& Settings, FAnimationSerializer *InAnimationSerializer = nullptr);
+	UE_API void InitInternal(USkeletalMeshComponent* InComponent, const FAnimationRecordingSettings& Settings, FAnimationSerializer *InAnimationSerializer = nullptr);
 
 public:
 	TWeakObjectPtr<USkeletalMeshComponent> SkelComp;
@@ -263,44 +265,45 @@ public:
 //////////////////////////////////////////////////////////////////////////
 // FAnimationRecorderManager
 
-struct SEQUENCERECORDER_API FAnimationRecorderManager
+struct FAnimationRecorderManager
 {
 public:
 	/** Singleton accessor */
-	static FAnimationRecorderManager& Get();
+	static UE_API FAnimationRecorderManager& Get();
 
 	/** Destructor */
-	virtual ~FAnimationRecorderManager();
+	UE_API virtual ~FAnimationRecorderManager();
 
 	/** Starts recording an animation. */
-	bool RecordAnimation(USkeletalMeshComponent* Component, const FString& AssetPath = FString(), const FString& AssetName = FString(), const FAnimationRecordingSettings& Settings = FAnimationRecordingSettings());
+	UE_API bool RecordAnimation(USkeletalMeshComponent* Component, const FString& AssetPath = FString(), const FString& AssetName = FString(), const FAnimationRecordingSettings& Settings = FAnimationRecordingSettings());
 
-	bool RecordAnimation(USkeletalMeshComponent* Component, UAnimSequence* Sequence, const FAnimationRecordingSettings& Settings = FAnimationRecordingSettings());
+	UE_API bool RecordAnimation(USkeletalMeshComponent* Component, UAnimSequence* Sequence, const FAnimationRecordingSettings& Settings = FAnimationRecordingSettings());
 	
-	bool RecordAnimation(USkeletalMeshComponent* Component, UAnimSequence* Sequence, FAnimationSerializer *InAnimationSerializer,  const FAnimationRecordingSettings& Settings = FAnimationRecordingSettings());
+	UE_API bool RecordAnimation(USkeletalMeshComponent* Component, UAnimSequence* Sequence, FAnimationSerializer *InAnimationSerializer,  const FAnimationRecordingSettings& Settings = FAnimationRecordingSettings());
 
-	bool IsRecording(USkeletalMeshComponent* Component);
+	UE_API bool IsRecording(USkeletalMeshComponent* Component);
 
-	bool IsRecording();
+	UE_API bool IsRecording();
 
-	UAnimSequence* GetCurrentlyRecordingSequence(USkeletalMeshComponent* Component);
-	float GetCurrentRecordingTime(USkeletalMeshComponent* Component);
-	void StopRecordingAnimation(USkeletalMeshComponent* Component, bool bShowMessage = true);
-	void StopRecordingAllAnimations();
-	const FTransform& GetInitialRootTransform(USkeletalMeshComponent* Component) const;
+	UE_API UAnimSequence* GetCurrentlyRecordingSequence(USkeletalMeshComponent* Component);
+	UE_API float GetCurrentRecordingTime(USkeletalMeshComponent* Component);
+	UE_API void StopRecordingAnimation(USkeletalMeshComponent* Component, bool bShowMessage = true);
+	UE_API void StopRecordingAllAnimations();
+	UE_API const FTransform& GetInitialRootTransform(USkeletalMeshComponent* Component) const;
 
-	void Tick(float DeltaTime);
+	UE_API void Tick(float DeltaTime);
 
-	void Tick(USkeletalMeshComponent* Component, float DeltaTime);
+	UE_API void Tick(USkeletalMeshComponent* Component, float DeltaTime);
 
-	void StopRecordingDeadAnimations(bool bShowMessage = true);
+	UE_API void StopRecordingDeadAnimations(bool bShowMessage = true);
 
 private:
 	/** Constructor, private - use Get() function */
-	FAnimationRecorderManager();
+	UE_API FAnimationRecorderManager();
 
 	TArray<FAnimRecorderInstance> RecorderInstances;
 
-	void HandleEndPIE(bool bSimulating);
+	UE_API void HandleEndPIE(bool bSimulating);
 };
 
+#undef UE_API
diff --git a/Engine/Source/Editor/SequenceRecorder/Public/ISequenceRecorderExtender.h b/Engine/Source/Editor/SequenceRecorder/Public/ISequenceRecorderExtender.h
--- a/Engine/Source/Editor/SequenceRecorder/Public/ISequenceRecorderExtender.h
+++ b/Engine/Source/Editor/SequenceRecorder/Public/ISequenceRecorderExtender.h
@@ -11,7 +11,7 @@ class SCompoundWidget;
 class UObject;
 class USequenceRecordingBase;
 
-class SEQUENCERECORDER_API ISequenceRecorderExtender : public TSharedFromThis<ISequenceRecorderExtender>
+class ISequenceRecorderExtender : public TSharedFromThis<ISequenceRecorderExtender>
 {
 public:
 	virtual ~ISequenceRecorderExtender() {}
diff --git a/Engine/Source/Editor/SequenceRecorder/Public/MovieScenePropertyRecorder.h b/Engine/Source/Editor/SequenceRecorder/Public/MovieScenePropertyRecorder.h
--- a/Engine/Source/Editor/SequenceRecorder/Public/MovieScenePropertyRecorder.h
+++ b/Engine/Source/Editor/SequenceRecorder/Public/MovieScenePropertyRecorder.h
@@ -19,6 +19,8 @@
 #include "UObject/WeakObjectPtr.h"
 #include "UObject/WeakObjectPtrTemplates.h"
 
+#define UE_API SEQUENCERECORDER_API
+
 class UObject;
 struct FGuid;
 
@@ -45,7 +47,7 @@ struct FPropertyKey
 
 /** Recorder for a simple property of type PropertyType */
 template <typename PropertyType>
-class SEQUENCERECORDER_API FMovieScenePropertyRecorder : public IMovieScenePropertyRecorder
+class FMovieScenePropertyRecorder : public IMovieScenePropertyRecorder
 {
 public:
 	FMovieScenePropertyRecorder(const FTrackInstancePropertyBindings& InBinding)
@@ -142,7 +144,7 @@ private:
 };
 
 /** Recorder for a simple property of type enum */
-class SEQUENCERECORDER_API FMovieScenePropertyRecorderEnum : public IMovieScenePropertyRecorder
+class FMovieScenePropertyRecorderEnum : public IMovieScenePropertyRecorder
 {
 public:
 	FMovieScenePropertyRecorderEnum(const FTrackInstancePropertyBindings& InBinding)
@@ -209,16 +211,16 @@ private:
 	 * Helper function, specialized by type, used to check if we do capture-time filtering of keys based
 	 * on previous values
 	 */
-	bool ShouldAddNewKey(const int64& InNewValue) const;
+	UE_API bool ShouldAddNewKey(const int64& InNewValue) const;
 
 	/** Helper function, specialized by type, used to add an appropriate section to the movie scene */
-	class UMovieSceneSection* AddSection(UObject* InObjectToRecord, class UMovieScene* InMovieScene, const FGuid& InGuid, float InTime);
+	UE_API class UMovieSceneSection* AddSection(UObject* InObjectToRecord, class UMovieScene* InMovieScene, const FGuid& InGuid, float InTime);
 
 	/** Helper function, specialized by type, used to add keys to the movie scene section at Finalize() time */
-	void AddKeyToSection(UMovieSceneSection* InSection, const FPropertyKey<int64>& InKey);
+	UE_API void AddKeyToSection(UMovieSceneSection* InSection, const FPropertyKey<int64>& InKey);
 
 	/** Helper function, specialized by type, used to reduce keys */
-	void ReduceKeys(UMovieSceneSection* InSection);
+	UE_API void ReduceKeys(UMovieSceneSection* InSection);
 
 private:
 	/** Binding for this property */
@@ -233,3 +235,5 @@ private:
 	/** Previous value we use to establish whether we should key */
 	int64 PreviousValue;
 };
+
+#undef UE_API
diff --git a/Engine/Source/Editor/SequenceRecorder/Public/SequenceRecorderActorGroup.h b/Engine/Source/Editor/SequenceRecorder/Public/SequenceRecorderActorGroup.h
--- a/Engine/Source/Editor/SequenceRecorder/Public/SequenceRecorderActorGroup.h
+++ b/Engine/Source/Editor/SequenceRecorder/Public/SequenceRecorderActorGroup.h
@@ -7,11 +7,13 @@
 #include "GameFramework/Actor.h"
 #include "SequenceRecorderActorGroup.generated.h"
 
+#define UE_API SEQUENCERECORDER_API
+
 // Forward Declares
 class UActorRecording;
 
-UCLASS()
-class SEQUENCERECORDER_API USequenceRecorderActorGroup : public UObject
+UCLASS(MinimalAPI)
+class USequenceRecorderActorGroup : public UObject
 {
 	GENERATED_BODY()
 public:
@@ -55,14 +57,14 @@ public:
 	UPROPERTY(VisibleAnywhere, Category = "Recording Groups")
 	TArray<TObjectPtr<UActorRecording>> RecordedActors;
 
-	void PostEditChangeChainProperty(struct FPropertyChangedChainEvent& PropertyChangedEvent) override;
+	UE_API void PostEditChangeChainProperty(struct FPropertyChangedChainEvent& PropertyChangedEvent) override;
 };
 
 /**
  * 
  */
-UCLASS(hidedropdown, NotBlueprintable)
-class SEQUENCERECORDER_API ASequenceRecorderGroup : public AActor
+UCLASS(MinimalAPI, hidedropdown, NotBlueprintable)
+class ASequenceRecorderGroup : public AActor
 {
 	GENERATED_BODY()
 public:
@@ -97,3 +99,5 @@ public:
 	UPROPERTY()
 	TArray<TObjectPtr<USequenceRecorderActorGroup>> ActorGroups;
 };
+
+#undef UE_API
diff --git a/Engine/Source/Editor/SequenceRecorder/Public/SequenceRecorderBlueprintLibrary.h b/Engine/Source/Editor/SequenceRecorder/Public/SequenceRecorderBlueprintLibrary.h
--- a/Engine/Source/Editor/SequenceRecorder/Public/SequenceRecorderBlueprintLibrary.h
+++ b/Engine/Source/Editor/SequenceRecorder/Public/SequenceRecorderBlueprintLibrary.h
@@ -7,10 +7,12 @@
 #include "Kismet/BlueprintFunctionLibrary.h"
 #include "SequenceRecorderBlueprintLibrary.generated.h"
 
+#define UE_API SEQUENCERECORDER_API
+
 class AActor;
 
-UCLASS(meta=(ScriptName="SequenceRecorderLibrary"))
-class SEQUENCERECORDER_API USequenceRecorderBlueprintLibrary : public UBlueprintFunctionLibrary
+UCLASS(MinimalAPI, meta=(ScriptName="SequenceRecorderLibrary"))
+class USequenceRecorderBlueprintLibrary : public UBlueprintFunctionLibrary
 {
 	GENERATED_BODY()
 
@@ -19,13 +21,15 @@ class SEQUENCERECORDER_API USequenceRecorderBlueprintLibrary : public UBlueprint
 	 * @param	ActorsToRecord	The actors to record
 	 */
 	UFUNCTION(BlueprintCallable, Category="Sequence Recording")
-	static void StartRecordingSequence(const TArray<AActor*>& ActorsToRecord);
+	static UE_API void StartRecordingSequence(const TArray<AActor*>& ActorsToRecord);
 
 	/** Are we currently recording a sequence */
 	UFUNCTION(BlueprintPure, Category="Sequence Recording")
-	static bool IsRecordingSequence();
+	static UE_API bool IsRecordingSequence();
 
 	/** Stop recording the current sequence, if any */
 	UFUNCTION(BlueprintCallable, Category="Sequence Recording")
-	static void StopRecordingSequence();
+	static UE_API void StopRecordingSequence();
 };
+
+#undef UE_API
diff --git a/Engine/Source/Editor/SequenceRecorder/Public/SequenceRecorderSettings.h b/Engine/Source/Editor/SequenceRecorder/Public/SequenceRecorderSettings.h
--- a/Engine/Source/Editor/SequenceRecorder/Public/SequenceRecorderSettings.h
+++ b/Engine/Source/Editor/SequenceRecorder/Public/SequenceRecorderSettings.h
@@ -13,6 +13,8 @@
 #include "SequenceRecorderActorFilter.h"
 #include "SequenceRecorderSettings.generated.h"
 
+#define UE_API SEQUENCERECORDER_API
+
 class ALevelSequenceActor;
 class ULevelSequence;
 
@@ -80,12 +82,12 @@ struct FSettingsForActorClass
 	bool bRecordToPossessable = false;
 };
 
-UCLASS(config=Editor)
-class SEQUENCERECORDER_API USequenceRecorderSettings : public UObject
+UCLASS(MinimalAPI, config=Editor)
+class USequenceRecorderSettings : public UObject
 {
 	GENERATED_UCLASS_BODY()
 
-	virtual void PostEditChangeChainProperty(struct FPropertyChangedChainEvent& PropertyChangedEvent) override;
+	UE_API virtual void PostEditChangeChainProperty(struct FPropertyChangedChainEvent& PropertyChangedEvent) override;
 
 public:
 	/** Whether to create a level sequence when recording. Actors and animations will be inserted into this sequence */
@@ -192,3 +194,5 @@ public:
 	UPROPERTY(Config, EditAnywhere, AdvancedDisplay, Category = "Sequence Recording")
 	TArray<FSettingsForActorClass> PerActorSettings;
 };
+
+#undef UE_API
diff --git a/Engine/Source/Editor/SequenceRecorder/Public/SequenceRecordingBase.h b/Engine/Source/Editor/SequenceRecorder/Public/SequenceRecordingBase.h
--- a/Engine/Source/Editor/SequenceRecorder/Public/SequenceRecordingBase.h
+++ b/Engine/Source/Editor/SequenceRecorder/Public/SequenceRecordingBase.h
@@ -11,8 +11,8 @@
 class ULevelSequence;
 class UMovieScene;
 
-UCLASS(abstract)
-class SEQUENCERECORDER_API USequenceRecordingBase : public UObject
+UCLASS(MinimalAPI, abstract)
+class USequenceRecordingBase : public UObject
 {
 	GENERATED_BODY()
 
diff --git a/Engine/Source/Editor/Sequencer/Private/SequencerKeyActor.h b/Engine/Source/Editor/Sequencer/Private/SequencerKeyActor.h
--- a/Engine/Source/Editor/Sequencer/Private/SequencerKeyActor.h
+++ b/Engine/Source/Editor/Sequencer/Private/SequencerKeyActor.h
@@ -5,17 +5,19 @@
 #include "Engine/StaticMeshActor.h"
 #include "SequencerKeyActor.generated.h"
 
-UCLASS()
-class SEQUENCER_API ASequencerKeyActor : public AActor
+#define UE_API SEQUENCER_API
+
+UCLASS(MinimalAPI)
+class ASequencerKeyActor : public AActor
 {
 	GENERATED_BODY()
 
 public:
 
-	ASequencerKeyActor();
+	UE_API ASequencerKeyActor();
 
 	/** AActor overrides */
-	virtual void PostEditMove(bool bFinished) override;
+	UE_API virtual void PostEditMove(bool bFinished) override;
 	virtual bool IsEditorOnly() const final
 	{
 		return true;
@@ -23,7 +25,7 @@ public:
 	/** AActor */
 
 	/** Set the track section and time associated with this key */
-	void SetKeyData(class UMovieScene3DTransformSection* NewTrackSection, float NewKeyTime);
+	UE_API void SetKeyData(class UMovieScene3DTransformSection* NewTrackSection, float NewKeyTime);
 
 	/** Return the actor associated with this key */
 	AActor* GetAssociatedActor() const
@@ -40,7 +42,7 @@ public:
 protected:
 
 	/** Push the data from a key actor change back to Sequencer */
-	void PropagateKeyChange();
+	UE_API void PropagateKeyChange();
 
 private:
 
@@ -59,4 +61,6 @@ private:
 	/** The time this key is associated with in Sequencer */
 	UPROPERTY()
 	float KeyTime;
-};
\ No newline at end of file
+};
+
+#undef UE_API
diff --git a/Engine/Source/Editor/Sequencer/Public/Filters/ISequencerFilterBar.h b/Engine/Source/Editor/Sequencer/Public/Filters/ISequencerFilterBar.h
--- a/Engine/Source/Editor/Sequencer/Public/Filters/ISequencerFilterBar.h
+++ b/Engine/Source/Editor/Sequencer/Public/Filters/ISequencerFilterBar.h
@@ -1,4 +1,4 @@
-﻿// Copyright Epic Games, Inc. All Rights Reserved.
+// Copyright Epic Games, Inc. All Rights Reserved.
 
 #pragma once
 
@@ -21,73 +21,73 @@ public:
 
 	virtual ~ISequencerFilterBar() = default;
 
-	SEQUENCER_API virtual FName GetIdentifier() const = 0;
+	virtual FName GetIdentifier() const = 0;
 
-	SEQUENCER_API virtual ISequencer& GetSequencer() const = 0;
+	virtual ISequencer& GetSequencer() const = 0;
 
-	SEQUENCER_API virtual TSharedPtr<FUICommandList> GetCommandList() const = 0;
+	virtual TSharedPtr<FUICommandList> GetCommandList() const = 0;
 
-	SEQUENCER_API virtual FString GetTextFilterString() const = 0;
-	SEQUENCER_API virtual void SetTextFilterString(const FString& InText) = 0;
+	virtual FString GetTextFilterString() const = 0;
+	virtual void SetTextFilterString(const FString& InText) = 0;
 
 	/** Returns true if the current filter bar text filter string contains the specified text expression.
 	 * The text expression must have key, operator, and value tokens. */
-	SEQUENCER_API virtual bool DoesTextFilterStringContainExpressionPair(const ISequencerTextFilterExpressionContext& InExpression) const = 0;
+	virtual bool DoesTextFilterStringContainExpressionPair(const ISequencerTextFilterExpressionContext& InExpression) const = 0;
 
-	SEQUENCER_API virtual bool AreFiltersMuted() const = 0;
-	SEQUENCER_API virtual void MuteFilters(const bool bInMute) = 0;
+	virtual bool AreFiltersMuted() const = 0;
+	virtual void MuteFilters(const bool bInMute) = 0;
 
-	SEQUENCER_API virtual bool CanResetFilters() const = 0;
-	SEQUENCER_API virtual void ResetFilters() = 0;
+	virtual bool CanResetFilters() const = 0;
+	virtual void ResetFilters() = 0;
 
-	SEQUENCER_API virtual bool HasAnyFilterActive(const bool bCheckTextFilter = true
+	virtual bool HasAnyFilterActive(const bool bCheckTextFilter = true
 		, const bool bInCheckHideIsolateFilter = true
 		, const bool bInCheckCommonFilters = true
 		, const bool bInCheckInternalFilters = true
 		, const bool bInCheckCustomTextFilters = true) const = 0;
 
-	SEQUENCER_API virtual bool HasAnyFilterEnabled() const = 0;
+	virtual bool HasAnyFilterEnabled() const = 0;
 
-	SEQUENCER_API virtual void RequestFilterUpdate() = 0;
+	virtual void RequestFilterUpdate() = 0;
 
-	SEQUENCER_API virtual void EnableAllFilters(const bool bInEnable, const TArray<FString>& InExceptionFilterNames) = 0;
+	virtual void EnableAllFilters(const bool bInEnable, const TArray<FString>& InExceptionFilterNames) = 0;
 
-	SEQUENCER_API virtual void ActivateCommonFilters(const bool bInActivate, const TArray<FString>& InExceptionFilterNames) = 0;
+	virtual void ActivateCommonFilters(const bool bInActivate, const TArray<FString>& InExceptionFilterNames) = 0;
 
-	SEQUENCER_API virtual bool AreAllEnabledFiltersActive(const bool bInActive, const TArray<FString> InExceptionFilterNames) const = 0;
-	SEQUENCER_API virtual void ActivateAllEnabledFilters(const bool bInActivate, const TArray<FString> InExceptionFilterNames) = 0;
+	virtual bool AreAllEnabledFiltersActive(const bool bInActive, const TArray<FString> InExceptionFilterNames) const = 0;
+	virtual void ActivateAllEnabledFilters(const bool bInActivate, const TArray<FString> InExceptionFilterNames) = 0;
 
-	SEQUENCER_API virtual bool IsFilterActiveByDisplayName(const FString& InFilterName) const = 0;
-	SEQUENCER_API virtual bool IsFilterEnabledByDisplayName(const FString& InFilterName) const = 0;
-	SEQUENCER_API virtual bool SetFilterActiveByDisplayName(const FString& InFilterName, const bool bInActive, const bool bInRequestFilterUpdate = true) = 0;
-	SEQUENCER_API virtual bool SetFilterEnabledByDisplayName(const FString& InFilterName, const bool bInEnabled, const bool bInRequestFilterUpdate = true) = 0;
+	virtual bool IsFilterActiveByDisplayName(const FString& InFilterName) const = 0;
+	virtual bool IsFilterEnabledByDisplayName(const FString& InFilterName) const = 0;
+	virtual bool SetFilterActiveByDisplayName(const FString& InFilterName, const bool bInActive, const bool bInRequestFilterUpdate = true) = 0;
+	virtual bool SetFilterEnabledByDisplayName(const FString& InFilterName, const bool bInEnabled, const bool bInRequestFilterUpdate = true) = 0;
 
-	SEQUENCER_API virtual TArray<FText> GetFilterDisplayNames() const = 0;
-	SEQUENCER_API virtual TArray<FText> GetCustomTextFilterNames() const = 0;
+	virtual TArray<FText> GetFilterDisplayNames() const = 0;
+	virtual TArray<FText> GetCustomTextFilterNames() const = 0;
 
-	SEQUENCER_API virtual int32 GetTotalDisplayNodeCount() const = 0;
-	SEQUENCER_API virtual int32 GetFilteredDisplayNodeCount() const  = 0;
+	virtual int32 GetTotalDisplayNodeCount() const = 0;
+	virtual int32 GetFilteredDisplayNodeCount() const  = 0;
 
-	SEQUENCER_API virtual const FTextFilterExpressionEvaluator& GetTextFilterExpressionEvaluator() const = 0;
-	SEQUENCER_API virtual TArray<TSharedRef<ISequencerTextFilterExpressionContext>> GetTextFilterExpressionContexts() const = 0;
+	virtual const FTextFilterExpressionEvaluator& GetTextFilterExpressionEvaluator() const = 0;
+	virtual TArray<TSharedRef<ISequencerTextFilterExpressionContext>> GetTextFilterExpressionContexts() const = 0;
 
-	SEQUENCER_API virtual void OpenTextExpressionHelp() = 0;
-	SEQUENCER_API virtual void SaveCurrentFilterSetAsCustomTextFilter() = 0;
-	SEQUENCER_API virtual void CreateNewTextFilter() = 0;
+	virtual void OpenTextExpressionHelp() = 0;
+	virtual void SaveCurrentFilterSetAsCustomTextFilter() = 0;
+	virtual void CreateNewTextFilter() = 0;
 
 	/** Event called when the visibility or layout state has changed */
-	SEQUENCER_API virtual FOnFilterBarStateChanged& OnStateChanged() = 0;
+	virtual FOnFilterBarStateChanged& OnStateChanged() = 0;
 
 	/** Event called to request that the subscribe update its view since the filters state has changed */
-	SEQUENCER_API virtual FSimpleMulticastDelegate& OnRequestUpdate() = 0;
+	virtual FSimpleMulticastDelegate& OnRequestUpdate() = 0;
 
-	SEQUENCER_API virtual bool ShouldShowFilterBarWidget() const = 0;
+	virtual bool ShouldShowFilterBarWidget() const = 0;
 
-	SEQUENCER_API virtual bool IsFilterBarVisible() const = 0;
-	SEQUENCER_API virtual void ToggleFilterBarVisibility() = 0;
+	virtual bool IsFilterBarVisible() const = 0;
+	virtual void ToggleFilterBarVisibility() = 0;
 
-	SEQUENCER_API virtual bool IsFilterBarLayout(const EFilterBarLayout InLayout) const = 0;
-	SEQUENCER_API virtual void SetToVerticalLayout() = 0;
-	SEQUENCER_API virtual void SetToHorizontalLayout() = 0;
-	SEQUENCER_API virtual void ToggleFilterBarLayout() = 0;
+	virtual bool IsFilterBarLayout(const EFilterBarLayout InLayout) const = 0;
+	virtual void SetToVerticalLayout() = 0;
+	virtual void SetToHorizontalLayout() = 0;
+	virtual void ToggleFilterBarLayout() = 0;
 };
diff --git a/Engine/Source/Editor/Sequencer/Public/Filters/ISequencerTextFilterExpressionContext.h b/Engine/Source/Editor/Sequencer/Public/Filters/ISequencerTextFilterExpressionContext.h
--- a/Engine/Source/Editor/Sequencer/Public/Filters/ISequencerTextFilterExpressionContext.h
+++ b/Engine/Source/Editor/Sequencer/Public/Filters/ISequencerTextFilterExpressionContext.h
@@ -1,4 +1,4 @@
-﻿// Copyright Epic Games, Inc. All Rights Reserved.
+// Copyright Epic Games, Inc. All Rights Reserved.
 
 #pragma once
 
@@ -12,7 +12,7 @@ enum class SEQUENCER_API ESequencerTextFilterValueType : uint8
 	Integer
 };
 
-struct SEQUENCER_API FSequencerTextFilterKeyword
+struct FSequencerTextFilterKeyword
 {
 	FString Keyword;
 	FText Description;
@@ -22,13 +22,13 @@ struct SEQUENCER_API FSequencerTextFilterKeyword
 class ISequencerTextFilterExpressionContext : public ITextFilterExpressionContext
 {
 public:
-	SEQUENCER_API virtual TSet<FName> GetKeys() const = 0;
+	virtual TSet<FName> GetKeys() const = 0;
 
-	SEQUENCER_API virtual ESequencerTextFilterValueType GetValueType() const = 0;
-	SEQUENCER_API virtual TArray<FSequencerTextFilterKeyword> GetValueKeywords() const { return {}; }
+	virtual ESequencerTextFilterValueType GetValueType() const = 0;
+	virtual TArray<FSequencerTextFilterKeyword> GetValueKeywords() const { return {}; }
 
-	SEQUENCER_API virtual FText GetDescription() const = 0;
-	SEQUENCER_API virtual FText GetCategory() const { return FText::GetEmpty(); }
+	virtual FText GetDescription() const = 0;
+	virtual FText GetCategory() const { return FText::GetEmpty(); }
 
 	static bool CompareFStringForExactBool(const FTextFilterString& InValue, const bool bInPassedFilter)
 	{
