TakeRecorder: ActorRecordingSettings won't create settings objects in the constructor any more, improving struct validation times.

diff --git a/Engine/Source/Editor/SequenceRecorder/Private/ActorRecordingSettings.cpp b/Engine/Source/Editor/SequenceRecorder/Private/ActorRecordingSettings.cpp
--- a/Engine/Source/Editor/SequenceRecorder/Private/ActorRecordingSettings.cpp
+++ b/Engine/Source/Editor/SequenceRecorder/Private/ActorRecordingSettings.cpp
@@ -13,13 +13,11 @@
 
 FActorRecordingSettings::FActorRecordingSettings()
 {
-	CreateSettingsObjectsFromFactory();
 }
 
 FActorRecordingSettings::FActorRecordingSettings(UObject* InOuter)
 {
 	Outer = InOuter;
-	CreateSettingsObjectsFromFactory();
 }
 
 void FActorRecordingSettings::CreateSettingsObjectsFromFactory()
@@ -34,3 +32,11 @@ void FActorRecordingSettings::CreateSettingsObjectsFromFactory()
 		}
 	}
 }
+
+void FActorRecordingSettings::CreateSettingsObjectsIfNeeded() const
+{
+	if (Settings.Num() == 0)
+	{
+		const_cast<FActorRecordingSettings*>(this)->CreateSettingsObjectsFromFactory();
+	}
+}
diff --git a/Engine/Source/Editor/SequenceRecorder/Public/ActorRecordingSettings.h b/Engine/Source/Editor/SequenceRecorder/Public/ActorRecordingSettings.h
--- a/Engine/Source/Editor/SequenceRecorder/Public/ActorRecordingSettings.h
+++ b/Engine/Source/Editor/SequenceRecorder/Public/ActorRecordingSettings.h
@@ -29,6 +29,7 @@ struct FActorRecordingSettings
 	template <typename SettingsType>
 	SettingsType* GetSettingsObject() const
 	{
+		CreateSettingsObjectsIfNeeded();
 		for (UObject* SettingsObject : Settings)
 		{
 			if (SettingsType* TypedSettingsObject = Cast<SettingsType>(SettingsObject))
@@ -40,6 +41,9 @@ struct FActorRecordingSettings
 		return nullptr;
 	}
 
+private:
+	UE_API void CreateSettingsObjectsIfNeeded() const;
+
 private:
 	/** External settings objects for recorders that supply them. Displayed via a details customization  */
 	UPROPERTY(EditAnywhere, Category = "Actor Recording")
