Take Recorder - Add IAnimationRecorder and IAnimRecorderInstance interfaces to a new SequenceRuntimeRecorder module. This opens up opportunities for other recorder functionality such as runtime.

Update FAnimationRecorder and FAnimRecorderInstance classes to inherit from these interfaces. Functionality is unchanged.

The SequenceRecorder module is not changed, but a new SequenceRuntimeRecorder module is added which contains the new interfaces to allow for extending animation recording to runtime as well.
Additionally moved the FTimecodeBoneMethod and ETimecodeBoneMode types into the SequenceRuntimeRecorder module and added type redirects to BaseEngine.ini.

Note that due to moving the FTimecodeBoneMethod and ETimecodeBoneMode types from editor to runtime, any modules using these types will now need to link against SequenceRuntimeRecorder.

#Synced-CL 52673380

diff --git a/Engine/Config/BaseEngine.ini b/Engine/Config/BaseEngine.ini
--- a/Engine/Config/BaseEngine.ini
+++ b/Engine/Config/BaseEngine.ini
@@ -1425,6 +1425,10 @@ ClassName=/Script/OnlineSubsystemUtils.OnlineEngineInterfaceImpl
 +StructRedirects=(OldName="/Script/Water.ShallowWaterSimulationGrid",NewName="/Script/Water.ShallowWaterSimulationGrid_DEPRECATED")
 +PropertyRedirects=(OldName="/Script/Water.BakedShallowWaterSimulationComponent.SimulationData",NewName="/Script/Water.BakedShallowWaterSimulationComponent.SimulationData_DEPRECATED")
 
+; SequenceRecorder → SequenceRuntimeRecorder type moves (5.8)
++StructRedirects=(OldName="/Script/SequenceRecorder.TimecodeBoneMethod",NewName="/Script/SequenceRuntimeRecorder.TimecodeBoneMethod")
++EnumRedirects=(OldName="/Script/SequenceRecorder.ETimecodeBoneMode",NewName="/Script/SequenceRuntimeRecorder.ETimecodeBoneMode")
+
 ; End of releases redirects
 
 ; Fixing the output of tags for blueprints to be reset.
diff --git a/Engine/Plugins/VirtualProduction/Takes/Source/TakeTrackRecorders/TakeTrackRecorders.Build.cs b/Engine/Plugins/VirtualProduction/Takes/Source/TakeTrackRecorders/TakeTrackRecorders.Build.cs
--- a/Engine/Plugins/VirtualProduction/Takes/Source/TakeTrackRecorders/TakeTrackRecorders.Build.cs
+++ b/Engine/Plugins/VirtualProduction/Takes/Source/TakeTrackRecorders/TakeTrackRecorders.Build.cs
@@ -16,6 +16,7 @@ public class TakeTrackRecorders : ModuleRules
 				"TakesCore",
 				"TimeManagement",
 				"SequenceRecorder", // For access to the individual UProperty type recorders
+				"SequenceRuntimeRecorder", // For FTimecodeBoneMethod and ETimecodeBoneMode
 				"SerializedRecorderInterface",
 				"UnrealEd",
 				"Sequencer"
diff --git a/Engine/Source/Editor/SequenceRecorder/Private/AnimationRecorder.cpp b/Engine/Source/Editor/SequenceRecorder/Private/AnimationRecorder.cpp
--- a/Engine/Source/Editor/SequenceRecorder/Private/AnimationRecorder.cpp
+++ b/Engine/Source/Editor/SequenceRecorder/Private/AnimationRecorder.cpp
@@ -26,11 +26,8 @@
 #include "Animation/BuiltInAttributeTypes.h"
 #include "AnimationRecorderParameters.h"
 
-#include UE_INLINE_GENERATED_CPP_BY_NAME(AnimationRecorder)
-
 #define LOCTEXT_NAMESPACE "FAnimationRecorder"
 
-
 static TAutoConsoleVariable<int32> CVarKeepNotifyAndCurvesOnAnimationRecord(
 	TEXT("a.KeepNotifyAndCurvesOnAnimationRecord"),
 	1,
diff --git a/Engine/Source/Editor/SequenceRecorder/Public/AnimationRecorder.h b/Engine/Source/Editor/SequenceRecorder/Public/AnimationRecorder.h
--- a/Engine/Source/Editor/SequenceRecorder/Public/AnimationRecorder.h
+++ b/Engine/Source/Editor/SequenceRecorder/Public/AnimationRecorder.h
@@ -12,7 +12,9 @@
 #include "Serializers/MovieSceneAnimationSerialization.h"
 #include "Misc/QualifiedFrameTime.h"
 #include "Animation/AnimTypes.h"
-#include "AnimationRecorder.generated.h"
+#include "TimecodeBoneMethod.h"
+#include "IAnimationRecorder.h"
+#include "IAnimRecorderInstance.h"
 
 #define UE_API SEQUENCERECORDER_API
 
@@ -24,51 +26,13 @@ class USkeletalMeshComponent;
 
 DECLARE_LOG_CATEGORY_EXTERN(AnimationSerialization, Verbose, All);
 
-UENUM(BlueprintType)
-enum class ETimecodeBoneMode : uint8
-{
-	All,
-	Root,
-	UserDefined,
-	MAX UMETA(Hidden)
-};
-
-USTRUCT(BlueprintType)
-struct FTimecodeBoneMethod
-{
-	GENERATED_USTRUCT_BODY()
-
-	/** Default constructor, initializing with default values */
-	FTimecodeBoneMethod() : BoneMode(ETimecodeBoneMode::Root) { }
-
-	/** The timecode bone mode */
-	UPROPERTY(EditAnywhere, BlueprintReadWrite, Category="Timecode")
-	ETimecodeBoneMode BoneMode;
-
-	/** Name of the bone to assign timecode values to */
-	UPROPERTY(EditAnywhere, BlueprintReadWrite, Category="Timecode")
-	FName BoneName;
-};
-
-struct FProcessRecordedTimeParams
-{
-	FString HoursName;
-	FString MinutesName;
-	FString SecondsName;
-	FString FramesName;
-	FString SubFramesName;
-
-	// Optionally support writing the Timecode Rate into time data.
-	TOptional<FString> RateName;
-
-	FString SlateName;
-	FString Slate;
-};
 //////////////////////////////////////////////////////////////////////////
 // FAnimationRecorder
 
 // records the mesh pose to animation input
-struct FAnimationRecorder : public FGCObject
+struct FAnimationRecorder 
+	: public UE::AnimationRecording::IAnimationRecorder
+	, public FGCObject
 {
 private:
 	/** Frame count used to signal an unbounded animation */
@@ -133,7 +97,7 @@ private:
 
 public:
 	UE_API FAnimationRecorder();
-	UE_API ~FAnimationRecorder();
+	UE_API virtual ~FAnimationRecorder() override;
 
 	// FGCObject interface start
 	UE_API virtual void AddReferencedObjects(FReferenceCollector& Collector) override;
@@ -147,11 +111,22 @@ public:
 	UE_API bool TriggerRecordAnimation(USkeletalMeshComponent* Component);
 	UE_API bool TriggerRecordAnimation(USkeletalMeshComponent* Component, const FString& AssetPath, const FString& AssetName);
 
-	UE_API void StartRecord(USkeletalMeshComponent* Component, UAnimSequence* InAnimationObject);
-	UE_API UAnimSequence* StopRecord(bool bShowMessage);
-	UE_API void UpdateRecord(USkeletalMeshComponent* Component, float DeltaTime);
+	// IAnimationRecorder interface start
+	UE_API virtual void StartRecord(USkeletalMeshComponent* Component, UAnimSequence* InAnimationObject) override;
+	UE_API virtual UAnimSequence* StopRecord(bool bShowMessage) override;
+	UE_API virtual void UpdateRecord(USkeletalMeshComponent* Component, float DeltaTime) override;
+	virtual bool InRecording() const override { return AnimationObject != nullptr; }
+
+	/** Sets the function that should be used to get the current timecode for the current frame. */
+	UE_API virtual void SetCurrentFrameTimeGetter(TFunction<TOptional<FQualifiedFrameTime>()> InGetCurrentTimeFunction) override;
+
+	virtual const FTransform& GetInitialRootTransform() const override { return InitialRootTransform; }
+
+	/** Process any time data captured and apply it to the bones on the given SkeletalMeshComponent. */
+	UE_API virtual void ProcessRecordedTimes(UAnimSequence* AnimSequence, USkeletalMeshComponent* SkeletalMeshComponent, const FTimecodeBoneMethod& TimecodeBoneMethod, const FProcessRecordedTimeParams& TimecodeInfo) override;
+	// IAnimationRecorder interface end
+
 	UAnimSequence* GetAnimationObject() const { return AnimationObject; }
-	bool InRecording() const { return AnimationObject != nullptr; }
 	double GetTimeRecorded() const { return TimePassed; }
 
 	/** Sets a new sample rate & max length for this recorder. Don't call while recording. */
@@ -159,17 +134,9 @@ public:
 
 	UE_API bool SetAnimCompressionScheme(UAnimBoneCompressionSettings* Settings);
 
-	/** Sets the function that should be used to get the current timecode for the current frame. */
-	UE_API void SetCurrentFrameTimeGetter(TFunction<TOptional<FQualifiedFrameTime>()> InGetCurrentTimeFunction);
-
-	const FTransform& GetInitialRootTransform() const { return InitialRootTransform; }
-
 	UE_DEPRECATED(5.5, "Use the ProcessRecordedTimes method that takes a FProcessRecordedTimeParams struct.")
 	UE_API void ProcessRecordedTimes(UAnimSequence* AnimSequence, USkeletalMeshComponent* SkeletalMeshComponent, const FString& HoursName, const FString& MinutesName, const FString& SecondsName, const FString& FramesName, const FString& SubFramesName, const FString& SlateName, const FString& Slate, const FTimecodeBoneMethod& TimecodeBoneMethod);
 
-	/** Process any time data captured and apply it to the bones on the given SkeletalMeshComponent. */
-	UE_API void ProcessRecordedTimes(UAnimSequence* AnimSequence, USkeletalMeshComponent* SkeletalMeshComponent, const FTimecodeBoneMethod& TimecodeBoneMethod, const FProcessRecordedTimeParams& TimecodeInfo);
-
 	/** If true, it will record root to include LocalToWorld */
 	uint8 bRecordLocalToWorld :1;
 	/** If true, asset will be saved to disk after recording. If false, asset will remain in mem and can be manually saved. */
@@ -224,26 +191,29 @@ private:
 //////////////////////////////////////////////////////////////////////////
 // FAnimRecorderInstance
 
-struct FAnimRecorderInstance
+struct FAnimRecorderInstance : public UE::AnimationRecording::IAnimRecorderInstance
 {
 public:
 	UE_API FAnimRecorderInstance();
-	UE_API ~FAnimRecorderInstance();
+	UE_API virtual ~FAnimRecorderInstance() override;
 
 	UE_API void Init(USkeletalMeshComponent* InComponent, const FString& InAssetPath, const FString& InAssetName, const FAnimationRecordingSettings& InSettings);
 
-	UE_API void Init(USkeletalMeshComponent* InComponent, UAnimSequence* InSequence, FAnimationSerializer *InAnimationSerializer, const FAnimationRecordingSettings& InSettings);
+	
+	// IAnimRecorderInstance interface start
+	UE_API virtual void Init(USkeletalMeshComponent* InComponent, UAnimSequence* InSequence, FAnimationSerializer *InAnimationSerializer, const FAnimationRecordingSettings& InSettings) override;
+	UE_API virtual bool BeginRecording() override;
+	UE_API virtual void Update(float DeltaTime) override;
+	UE_API virtual void FinishRecording(bool bShowMessage = true) override;
 
-	UE_API bool BeginRecording();
-	UE_API void Update(float DeltaTime);
-	UE_API void FinishRecording(bool bShowMessage = true);
+	/** Process any time data captured and apply it to the bones on the given SkeletalMeshComponent. */
+	UE_API virtual void ProcessRecordedTimes(UAnimSequence* AnimSequence, USkeletalMeshComponent* SkeletalMeshComponent, const FTimecodeBoneMethod& TimecodeBoneMethod, const FProcessRecordedTimeParams& TimecodeInfo) override;
+	virtual TSharedPtr<UE::AnimationRecording::IAnimationRecorder> GetRecorder() const override { return Recorder; }
+	// IAnimRecorderInstance interface end
 
 	UE_DEPRECATED(5.5, "Use the ProcessRecordedTimes method that takes a FProcessRecordedTimeParams struct.")
 	UE_API void ProcessRecordedTimes(UAnimSequence* AnimSequence, USkeletalMeshComponent* SkeletalMeshComponent, const FString& HoursName, const FString& MinutesName, const FString& SecondsName, const FString& FramesName, const FString& SubFramesName, const FString& SlateName, const FString& Slate, const FTimecodeBoneMethod& TimecodeBoneMethod);
 
-	/** Process any time data captured and apply it to the bones on the given SkeletalMeshComponent. */
-	UE_API void ProcessRecordedTimes(UAnimSequence* AnimSequence, USkeletalMeshComponent* SkeletalMeshComponent, const FTimecodeBoneMethod& TimecodeBoneMethod, const FProcessRecordedTimeParams& TimecodeInfo);
-
 private:
 	UE_API void InitInternal(USkeletalMeshComponent* InComponent, const FAnimationRecordingSettings& Settings, FAnimationSerializer *InAnimationSerializer = nullptr);
 
diff --git a/Engine/Source/Editor/SequenceRecorder/SequenceRecorder.Build.cs b/Engine/Source/Editor/SequenceRecorder/SequenceRecorder.Build.cs
--- a/Engine/Source/Editor/SequenceRecorder/SequenceRecorder.Build.cs
+++ b/Engine/Source/Editor/SequenceRecorder/SequenceRecorder.Build.cs
@@ -12,6 +12,7 @@ namespace UnrealBuildTool.Rules
 				{
 					"TimeManagement",
                     "SerializedRecorderInterface",
+					"SequenceRuntimeRecorder"
 
                 }
             );
diff --git a/Engine/Source/Runtime/SequenceRuntimeRecorder/Private/SequenceRuntimeRecorderModule.cpp b/Engine/Source/Runtime/SequenceRuntimeRecorder/Private/SequenceRuntimeRecorderModule.cpp
new file mode 100644
--- /dev/null
+++ b/Engine/Source/Runtime/SequenceRuntimeRecorder/Private/SequenceRuntimeRecorderModule.cpp
@@ -0,0 +1,5 @@
+// Copyright Epic Games, Inc. All Rights Reserved.
+
+#include "Modules/ModuleManager.h"
+
+IMPLEMENT_MODULE(FDefaultModuleImpl, SequenceRuntimeRecorder)
diff --git a/Engine/Source/Runtime/SequenceRuntimeRecorder/Public/IAnimRecorderInstance.h b/Engine/Source/Runtime/SequenceRuntimeRecorder/Public/IAnimRecorderInstance.h
new file mode 100644
--- /dev/null
+++ b/Engine/Source/Runtime/SequenceRuntimeRecorder/Public/IAnimRecorderInstance.h
@@ -0,0 +1,45 @@
+// Copyright Epic Games, Inc. All Rights Reserved.
+
+#pragma once
+
+#include "Containers/UnrealString.h"
+#include "IAnimationRecorder.h"
+#include "Templates/SharedPointer.h"
+#include "UObject/WeakObjectPtr.h"
+
+//////////////////////////////////////////////////////////////////////////
+// IAnimRecorderInstance
+
+namespace UE::AnimationRecording
+{
+
+/**
+ * Base interface for animation recorder instances.
+ * Manages the lifecycle of a single recording session for a skeletal mesh component.
+ */
+struct IAnimRecorderInstance
+{
+public:
+	virtual ~IAnimRecorderInstance() = default;
+
+	/** Initializes the instance to record into an existing animation sequence. */
+	virtual void Init(USkeletalMeshComponent* InComponent, UAnimSequence* InSequence, FAnimationSerializer* InAnimationSerializer, const FAnimationRecordingSettings& InSettings) = 0;
+
+	/** Starts the recording session. Returns true on success. */
+	virtual bool BeginRecording() = 0;
+
+	/** Records a single frame of pose data. Call each tick while recording. */
+	virtual void Update(float DeltaTime) = 0;
+
+	/** Stops recording and finalizes the animation sequence. */
+	virtual void FinishRecording(bool bShowMessage = true) = 0;
+
+	/** Process any time data captured and apply it to the bones on the given SkeletalMeshComponent. */
+	virtual void ProcessRecordedTimes(
+		UAnimSequence* AnimSequence, USkeletalMeshComponent* SkeletalMeshComponent,
+		const FTimecodeBoneMethod& TimecodeBoneMethod, const FProcessRecordedTimeParams& TimecodeInfo) = 0;
+
+	virtual TSharedPtr<UE::AnimationRecording::IAnimationRecorder> GetRecorder() const = 0;
+};
+
+} // namespace UE::AnimationRecording
diff --git a/Engine/Source/Runtime/SequenceRuntimeRecorder/Public/IAnimationRecorder.h b/Engine/Source/Runtime/SequenceRuntimeRecorder/Public/IAnimationRecorder.h
new file mode 100644
--- /dev/null
+++ b/Engine/Source/Runtime/SequenceRuntimeRecorder/Public/IAnimationRecorder.h
@@ -0,0 +1,73 @@
+// Copyright Epic Games, Inc. All Rights Reserved.
+
+#pragma once
+
+#include "Containers/UnrealString.h"
+#include "Logging/LogCategory.h"
+#include "Math/Transform.h"
+#include "Misc/FrameRate.h"
+#include "Misc/Optional.h"
+#include "Misc/QualifiedFrameTime.h"
+#include "Templates/Function.h"
+#include "TimecodeBoneMethod.h"
+#include "Serializers/MovieSceneAnimationSerialization.h"
+
+class UAnimBoneCompressionSettings;
+class UAnimSequence;
+class USkeletalMeshComponent;
+struct FAnimationRecordingSettings;
+
+struct FProcessRecordedTimeParams
+{
+	FString HoursName;
+	FString MinutesName;
+	FString SecondsName;
+	FString FramesName;
+	FString SubFramesName;
+
+	// Optionally support writing the Timecode Rate into time data.
+	TOptional<FString> RateName;
+
+	FString SlateName;
+	FString Slate;
+};
+
+//////////////////////////////////////////////////////////////////////////
+// IAnimationRecorder
+
+namespace UE::AnimationRecording
+{
+
+/**
+ * Base interface for animation recorders.
+ * Inheriting from this interface allows different animation recording implementations as required.
+ */
+struct IAnimationRecorder
+{
+	virtual ~IAnimationRecorder() = default;
+
+	/** Begins recording pose data from the component into the given animation sequence. */
+	virtual void StartRecord(USkeletalMeshComponent* Component, UAnimSequence* InAnimationObject) = 0;
+
+	/** Stops recording and finalizes the animation sequence. Returns the recorded sequence. */
+	virtual UAnimSequence* StopRecord(bool bShowMessage) = 0;
+
+	/** Records a single frame of pose data from the component. Call each tick while recording. */
+	virtual void UpdateRecord(USkeletalMeshComponent* Component, float DeltaTime) = 0;
+
+	/** Returns true if a recording is currently in progress. */
+	virtual bool InRecording() const = 0;
+
+	/** Sets the function that should be used to get the current timecode for the current frame. */
+	virtual void SetCurrentFrameTimeGetter(TFunction<TOptional<FQualifiedFrameTime>()> InGetCurrentTimeFunction) = 0;
+
+	/** Returns the root bone transform captured at the start of recording. */
+	virtual const FTransform& GetInitialRootTransform() const = 0;
+
+	/** Process any time data captured and apply it to the bones on the given SkeletalMeshComponent. */
+	virtual void ProcessRecordedTimes(
+		UAnimSequence* AnimSequence, USkeletalMeshComponent* SkeletalMeshComponent,
+		const FTimecodeBoneMethod& TimecodeBoneMethod, const FProcessRecordedTimeParams& TimecodeInfo) = 0;
+};
+
+} // namespace UE::AnimationRecording
diff --git a/Engine/Source/Runtime/SequenceRuntimeRecorder/Public/TimecodeBoneMethod.h b/Engine/Source/Runtime/SequenceRuntimeRecorder/Public/TimecodeBoneMethod.h
new file mode 100644
--- /dev/null
+++ b/Engine/Source/Runtime/SequenceRuntimeRecorder/Public/TimecodeBoneMethod.h
@@ -0,0 +1,34 @@
+// Copyright Epic Games, Inc. All Rights Reserved.
+
+#pragma once
+
+#include "UObject/ObjectMacros.h"
+#include "UObject/NameTypes.h"
+
+#include "TimecodeBoneMethod.generated.h"
+
+UENUM(BlueprintType)
+enum class ETimecodeBoneMode : uint8
+{
+	All,
+	Root,
+	UserDefined,
+	MAX UMETA(Hidden)
+};
+
+USTRUCT(BlueprintType)
+struct FTimecodeBoneMethod
+{
+	GENERATED_USTRUCT_BODY()
+
+	/** Default constructor, initializing with default values */
+	FTimecodeBoneMethod() : BoneMode(ETimecodeBoneMode::Root) { }
+
+	/** The timecode bone mode */
+	UPROPERTY(EditAnywhere, BlueprintReadWrite, Category="Timecode")
+	ETimecodeBoneMode BoneMode;
+
+	/** Name of the bone to assign timecode values to */
+	UPROPERTY(EditAnywhere, BlueprintReadWrite, Category="Timecode")
+	FName BoneName;
+};
diff --git a/Engine/Source/Runtime/SequenceRuntimeRecorder/SequenceRuntimeRecorder.Build.cs b/Engine/Source/Runtime/SequenceRuntimeRecorder/SequenceRuntimeRecorder.Build.cs
new file mode 100644
--- /dev/null
+++ b/Engine/Source/Runtime/SequenceRuntimeRecorder/SequenceRuntimeRecorder.Build.cs
@@ -0,0 +1,31 @@
+// Copyright Epic Games, Inc. All Rights Reserved.
+
+namespace UnrealBuildTool.Rules
+{
+	public class SequenceRuntimeRecorder : ModuleRules
+	{
+		public SequenceRuntimeRecorder(ReadOnlyTargetRules Target) : base(Target)
+		{
+			PublicDependencyModuleNames.AddRange(
+				new string[]
+				{
+					"TimeManagement",
+                    "SerializedRecorderInterface",
+
+                }
+            );
+
+			PrivateDependencyModuleNames.AddRange(
+				new string[]
+				{
+					"Core",
+					"CoreUObject",
+					"SlateCore",
+					"Slate",
+					"InputCore",
+					"Engine",
+                }
+            );
+		}
+	}
+}
